MeshConversion: handle invalid UV triangles in DynamicMeshToMeshDescription
#rnx

diff --git a/Engine/Source/Runtime/MeshConversion/Private/DynamicMeshToMeshDescription.cpp b/Engine/Source/Runtime/MeshConversion/Private/DynamicMeshToMeshDescription.cpp
--- a/Engine/Source/Runtime/MeshConversion/Private/DynamicMeshToMeshDescription.cpp
+++ b/Engine/Source/Runtime/MeshConversion/Private/DynamicMeshToMeshDescription.cpp
@@ -19,7 +19,7 @@ namespace DynamicMeshToMeshDescriptionConversionHelper
 	template <typename OutAttributeType, int VecLen, typename InAttributeType>
 	void SetAttributesFromOverlay(
 		const FDynamicMesh3* MeshInArg, const FMeshDescription& MeshOutArg,
-		TVertexInstanceAttributesRef<OutAttributeType>& InstanceAttrib, const TDynamicMeshVectorOverlay<float, VecLen, InAttributeType>* Overlay, int AttribIndex=0)
+		TVertexInstanceAttributesRef<OutAttributeType>& InstanceAttrib, const TDynamicMeshVectorOverlay<float, VecLen, InAttributeType>* Overlay, OutAttributeType DefaultValue, int AttribIndex=0)
 	{
 		for (const FTriangleID TriangleID : MeshOutArg.Triangles().GetElementIDs())
 		{
@@ -27,10 +27,19 @@ namespace DynamicMeshToMeshDescriptionConversionHelper
 
 			int32 MeshInTriIdx = TriangleID.GetValue();
 			
-			FIndex3i OverlayVertIndices = Overlay->GetTriangle(MeshInTriIdx);
-			InstanceAttrib.Set(InstanceTri[0], AttribIndex, OutAttributeType(Overlay->GetElement(OverlayVertIndices.A)));
-			InstanceAttrib.Set(InstanceTri[1], AttribIndex, OutAttributeType(Overlay->GetElement(OverlayVertIndices.B)));
-			InstanceAttrib.Set(InstanceTri[2], AttribIndex, OutAttributeType(Overlay->GetElement(OverlayVertIndices.C)));
+			if (Overlay->IsSetTriangle(MeshInTriIdx))
+			{
+				FIndex3i OverlayVertIndices = Overlay->GetTriangle(MeshInTriIdx);
+				InstanceAttrib.Set(InstanceTri[0], AttribIndex, OutAttributeType(Overlay->GetElement(OverlayVertIndices.A)));
+				InstanceAttrib.Set(InstanceTri[1], AttribIndex, OutAttributeType(Overlay->GetElement(OverlayVertIndices.B)));
+				InstanceAttrib.Set(InstanceTri[2], AttribIndex, OutAttributeType(Overlay->GetElement(OverlayVertIndices.C)));
+			}
+			else
+			{
+				InstanceAttrib.Set(InstanceTri[0], AttribIndex, DefaultValue);
+				InstanceAttrib.Set(InstanceTri[1], AttribIndex, DefaultValue);
+				InstanceAttrib.Set(InstanceTri[2], AttribIndex, DefaultValue);
+			}
 		}
 	}
 }
@@ -80,7 +89,7 @@ void FDynamicMeshToMeshDescription::UpdateAttributes(const FDynamicMesh3* MeshIn
 			if (Overlay)
 			{
 				check(MeshIn->TriangleCount() == MeshOut.Triangles().Num())
-				DynamicMeshToMeshDescriptionConversionHelper::SetAttributesFromOverlay(MeshIn, MeshOut, InstanceAttrib, Overlay);
+				DynamicMeshToMeshDescriptionConversionHelper::SetAttributesFromOverlay(MeshIn, MeshOut, InstanceAttrib, Overlay, FVector3f::ZeroVector);
 			}
 			else
 			{
@@ -121,7 +130,7 @@ void FDynamicMeshToMeshDescription::UpdateAttributes(const FDynamicMesh3* MeshIn
 				{
 					const FDynamicMeshUVOverlay* UVOverlay = MeshIn->Attributes()->GetUVLayer(UVLayerIndex);
 					// update the vertex Attribute UVs
-					DynamicMeshToMeshDescriptionConversionHelper::SetAttributesFromOverlay(MeshIn, MeshOut, InstanceAttrib, UVOverlay, UVLayerIndex);
+					DynamicMeshToMeshDescriptionConversionHelper::SetAttributesFromOverlay(MeshIn, MeshOut, InstanceAttrib, UVOverlay, FVector2D::ZeroVector, UVLayerIndex);
 
 					// rebuild the shared UVs
 					FUVArray& UVArray = MeshOut.UVs(UVLayerIndex);
@@ -149,11 +158,23 @@ void FDynamicMeshToMeshDescription::UpdateAttributes(const FDynamicMesh3* MeshIn
 					for (const FTriangleID TriangleID : MeshOut.Triangles().GetElementIDs())
 					{
 						int TriID = TriangleID.GetValue(); // assumes the same TriIDs in both meshes.
-						FIndex3i ElIDs = UVOverlay->GetTriangle(TriID);
-						TArray<FUVID, TFixedAllocator<3>> MDTri; 
-						MDTri.Add(ElIDToUVIDMap[ElIDs[0]]);
-						MDTri.Add(ElIDToUVIDMap[ElIDs[1]]);
-						MDTri.Add(ElIDToUVIDMap[ElIDs[2]]);
+						TArray<FUVID, TFixedAllocator<3>> MDTri;
+						if (UVOverlay->IsSetTriangle(TriID))
+						{
+							FIndex3i ElIDs = UVOverlay->GetTriangle(TriID);
+							MDTri.Add(ElIDToUVIDMap[ElIDs[0]]);
+							MDTri.Add(ElIDToUVIDMap[ElIDs[1]]);
+							MDTri.Add(ElIDToUVIDMap[ElIDs[2]]);
+						}
+						else
+						{
+							for (int j = 0; j < 3; ++j)
+							{
+								FUVID UVID = UVArray.Add();
+								UVArray.GetAttributes().GetAttributesRef<FVector2D>(MeshAttribute::UV::UVCoordinate)[UVID] = FVector2D::ZeroVector;
+								MDTri.Add(UVID);
+							}
+						}
 						MeshOut.SetTriangleUVIndices(TriangleID, MDTri, UVLayerIndex);
 					}
 				}
