ModelingTools: Pre-transform vertex colors SRGB to Linear in BakeVertex.

#rnx

diff --git a/Engine/Plugins/Experimental/MeshModelingToolsetExp/Source/MeshModelingToolsExp/Private/BakeMeshAttributeVertexTool.cpp b/Engine/Plugins/Experimental/MeshModelingToolsetExp/Source/MeshModelingToolsExp/Private/BakeMeshAttributeVertexTool.cpp
--- a/Engine/Plugins/Experimental/MeshModelingToolsetExp/Source/MeshModelingToolsExp/Private/BakeMeshAttributeVertexTool.cpp
+++ b/Engine/Plugins/Experimental/MeshModelingToolsetExp/Source/MeshModelingToolsExp/Private/BakeMeshAttributeVertexTool.cpp
@@ -401,6 +401,7 @@ void UBakeMeshAttributeVertexTool::Shutdown(EToolShutdownType ShutdownType)
 			                                               "Bake Mesh Attribute Vertex"));
 			FConversionToMeshDescriptionOptions ConvertOptions;
 			ConvertOptions.SetToVertexColorsOnly();
+			ConvertOptions.bTransformVtxColorsSRGBToLinear = true;
 			UE::ToolTarget::CommitDynamicMeshUpdate(
 				Targets[0],
 				*PreviewMesh->GetMesh(),
diff --git a/Engine/Plugins/Experimental/MeshModelingToolsetExp/Source/ModelingComponents/Private/ConversionUtils/DynamicMeshViaMeshDescriptionUtil.cpp b/Engine/Plugins/Experimental/MeshModelingToolsetExp/Source/ModelingComponents/Private/ConversionUtils/DynamicMeshViaMeshDescriptionUtil.cpp
--- a/Engine/Plugins/Experimental/MeshModelingToolsetExp/Source/ModelingComponents/Private/ConversionUtils/DynamicMeshViaMeshDescriptionUtil.cpp
+++ b/Engine/Plugins/Experimental/MeshModelingToolsetExp/Source/ModelingComponents/Private/ConversionUtils/DynamicMeshViaMeshDescriptionUtil.cpp
@@ -28,6 +28,7 @@ void UE::Geometry::CommitDynamicMeshViaMeshDescription(
 	ConversionOptions.bUpdateTangents = CommitInfo.bTangentsChanged;
 	ConversionOptions.bUpdateUVs = CommitInfo.bUVsChanged;
 	ConversionOptions.bUpdateVtxColors = CommitInfo.bVertexColorsChanged;
+	ConversionOptions.bTransformVtxColorsSRGBToLinear = CommitInfo.bTransformVertexColorsSRGBToLinear;
 
 	MeshDescriptionCommitter.CommitMeshDescription([&CommitInfo, &ConversionOptions, &Mesh](const IMeshDescriptionCommitter::FCommitterParams& CommitParams)
 	{
diff --git a/Engine/Plugins/Experimental/MeshModelingToolsetExp/Source/ModelingComponents/Private/ModelingToolTargetUtil.cpp b/Engine/Plugins/Experimental/MeshModelingToolsetExp/Source/ModelingComponents/Private/ModelingToolTargetUtil.cpp
--- a/Engine/Plugins/Experimental/MeshModelingToolsetExp/Source/ModelingComponents/Private/ModelingToolTargetUtil.cpp
+++ b/Engine/Plugins/Experimental/MeshModelingToolsetExp/Source/ModelingComponents/Private/ModelingToolTargetUtil.cpp
@@ -271,6 +271,7 @@ UE::ToolTarget::EDynamicMeshUpdateResult UE::ToolTarget::CommitDynamicMeshUpdate
 		CommitInfo.bTangentsChanged = ConversionOptions.bUpdateTangents;
 		CommitInfo.bUVsChanged = ConversionOptions.bUpdateUVs;
 		CommitInfo.bVertexColorsChanged = ConversionOptions.bUpdateVtxColors;
+		CommitInfo.bTransformVertexColorsSRGBToLinear = ConversionOptions.bTransformVtxColorsSRGBToLinear;
 
 		DynamicMeshCommitter->CommitDynamicMesh(UpdatedMesh, CommitInfo);
 
diff --git a/Engine/Plugins/Experimental/MeshModelingToolsetExp/Source/ModelingComponents/Public/TargetInterfaces/DynamicMeshCommitter.h b/Engine/Plugins/Experimental/MeshModelingToolsetExp/Source/ModelingComponents/Public/TargetInterfaces/DynamicMeshCommitter.h
--- a/Engine/Plugins/Experimental/MeshModelingToolsetExp/Source/ModelingComponents/Public/TargetInterfaces/DynamicMeshCommitter.h
+++ b/Engine/Plugins/Experimental/MeshModelingToolsetExp/Source/ModelingComponents/Public/TargetInterfaces/DynamicMeshCommitter.h
@@ -33,7 +33,7 @@ public:
 	 */
 	struct MODELINGCOMPONENTS_API FDynamicMeshCommitInfo
 	{
-		/** Initializes everything to bInitValue */
+		/** Initializes each of the b*Changed members to bInitValue */
 		FDynamicMeshCommitInfo(bool bInitValue)
 		{
 			bPositionsChanged =
@@ -55,6 +55,13 @@ public:
 		bool bTangentsChanged = true;
 		bool bUVsChanged = true;
 		bool bVertexColorsChanged = true;
+
+		/**
+		 * Intentionally left out of the constructor. This is a different
+		 * parameter than the b*Changed members that augments how vertex
+		 * colors are transformed during the commit.
+		 */
+		bool bTransformVertexColorsSRGBToLinear = false;
 	};
 
 	virtual void CommitDynamicMesh(const UE::Geometry::FDynamicMesh3& Mesh)
diff --git a/Engine/Source/Runtime/GeometryCore/Public/Util/ColorConstants.h b/Engine/Source/Runtime/GeometryCore/Public/Util/ColorConstants.h
--- a/Engine/Source/Runtime/GeometryCore/Public/Util/ColorConstants.h
+++ b/Engine/Source/Runtime/GeometryCore/Public/Util/ColorConstants.h
@@ -3,6 +3,7 @@
 #pragma once
 
 #include "Math/Color.h"
+#include "MathUtil.h"
 
 namespace UE
 {
@@ -17,6 +18,36 @@ namespace Geometry
  */
 namespace LinearColors
 {
+	/** Apply an sRGB to Linear color transformation on a given color vector. */
+	template <typename VectorType>
+	void SRGBToLinear(VectorType& Color)
+	{
+		// Sourced from MPCDIUtils.ush
+		auto SRGBToLinearFloat = [](const float Color) -> float
+		{
+			return (Color <= 0.04045f) ? Color / 12.92f : FMathf::Pow((Color + 0.055f) / 1.055f, 2.4f);
+		};
+		
+		Color.X = SRGBToLinearFloat(Color.X);
+		Color.Y = SRGBToLinearFloat(Color.Y);
+		Color.Z = SRGBToLinearFloat(Color.Z);
+	}
+
+	/** Apply a Linear to sRGB color transformation on a given color vector. */
+	template <typename VectorType>
+	void LinearToSRGB(VectorType& Color)
+	{
+		// Sourced from MPCDIUtils.ush
+		auto LinearToSRGBFloat = [](const float Color) -> float
+		{
+			return (Color <= 0.0031308f) ? Color * 12.92f : 1.055f * FMathf::Pow(Color, 1.0f / 2.4f) - 0.055f;
+		};
+
+		Color.X = LinearToSRGBFloat(Color.X);
+		Color.Y = LinearToSRGBFloat(Color.Y);
+		Color.Z = LinearToSRGBFloat(Color.Z);
+	}
+	
 	template <typename VectorType>
 	VectorType MakeColor3f(float R, float G, float B)
 	{
diff --git a/Engine/Source/Runtime/MeshConversion/Private/DynamicMeshToMeshDescription.cpp b/Engine/Source/Runtime/MeshConversion/Private/DynamicMeshToMeshDescription.cpp
--- a/Engine/Source/Runtime/MeshConversion/Private/DynamicMeshToMeshDescription.cpp
+++ b/Engine/Source/Runtime/MeshConversion/Private/DynamicMeshToMeshDescription.cpp
@@ -7,6 +7,7 @@
 #include "DynamicMesh/DynamicVertexSkinWeightsAttribute.h"
 #include "MeshDescriptionBuilder.h"
 #include "DynamicMesh/MeshTangents.h"
+#include "Util/ColorConstants.h"
 
 #include "ExplicitUseGeometryMathTypes.h"		// using UE::Geometry::(math types)
 using namespace UE::Geometry;
@@ -335,6 +336,26 @@ void FDynamicMeshToMeshDescription::UpdateVertexColors(const FDynamicMesh3* Mesh
 	
 		FVector4f TriColors[3];
 		ColorOverlay->GetTriElements(t, TriColors[0], TriColors[1], TriColors[2]);
+
+		// There is inconsistency in how vertex colors are intended to be consumed in
+		// our shaders. Some shaders consume it as linear (ex. MeshPaint), others as SRGB which
+		// manually convert to linear in the shader.
+		//
+		// All StaticMeshes store vertex colors as an 8-bit FColor. In order to ensure a good
+		// distribution of float values across the 8-bit range, the StaticMesh build always
+		// encodes FColors as SRGB.
+		//
+		// Until there is some defined gamma space convention for vertex colors in our shaders,
+		// we provide this option to pre-transform our linear float colors with an SRGBToLinear
+		// conversion to counteract the StaticMesh build LinearToSRGB conversion. This is how
+		// MeshPaint ensures linear vertex colors in the shaders.
+		if (ConversionOptions.bTransformVtxColorsSRGBToLinear)
+		{
+			LinearColors::SRGBToLinear(TriColors[0]);
+			LinearColors::SRGBToLinear(TriColors[1]);
+			LinearColors::SRGBToLinear(TriColors[2]);
+		}
+		
 		TArrayView<const FVertexInstanceID> InstanceIDs = MeshOut.GetTriangleVertexInstances(FTriangleID(t));
 		for (int32 i = 0; i < 3; ++i)
 		{
diff --git a/Engine/Source/Runtime/MeshConversion/Public/MeshConversionOptions.h b/Engine/Source/Runtime/MeshConversion/Public/MeshConversionOptions.h
--- a/Engine/Source/Runtime/MeshConversion/Public/MeshConversionOptions.h
+++ b/Engine/Source/Runtime/MeshConversion/Public/MeshConversionOptions.h
@@ -27,6 +27,8 @@ public:
 	/** Should Vertex Colors of MeshDescription be updated, if available and relevant */
 	bool bUpdateVtxColors = false;
 
+	/** Should Vertex Colors of MeshDescription be transformed from SRGB to Linear */
+	bool bTransformVtxColorsSRGBToLinear = false;
 
 	//
 	// utility functions for common configuration cases
