DynamicMesh: Fixed issue where new UVChannels would cause a crash when introduced during UpdateAttributes

#rnx

diff --git a/Engine/Source/Runtime/MeshConversion/Private/DynamicMeshToMeshDescription.cpp b/Engine/Source/Runtime/MeshConversion/Private/DynamicMeshToMeshDescription.cpp
--- a/Engine/Source/Runtime/MeshConversion/Private/DynamicMeshToMeshDescription.cpp
+++ b/Engine/Source/Runtime/MeshConversion/Private/DynamicMeshToMeshDescription.cpp
@@ -126,7 +126,8 @@ void FDynamicMeshToMeshDescription::UpdateAttributes(const FDynamicMesh3* MeshIn
 				check(MeshIn->TriangleCount() == MeshOut.Triangles().Num())
 				int32 NumLayers = MeshIn->Attributes()->NumUVLayers();			
 				MeshOut.SetNumUVChannels(NumLayers); // This resets MeshDescription's internal TriangleUV array
-				
+				InstanceAttrib.SetNumChannels(NumLayers);
+
 				for (int UVLayerIndex = 0; UVLayerIndex < NumLayers; UVLayerIndex++)
 				{
 					const FDynamicMeshUVOverlay* UVOverlay = MeshIn->Attributes()->GetUVLayer(UVLayerIndex);
