MeshConversion: move non-perf-critical functions from header to cpp to help w/ debugging
#rnx

diff --git a/Engine/Source/Runtime/MeshConversion/Private/DynamicMeshToMeshDescription.cpp b/Engine/Source/Runtime/MeshConversion/Private/DynamicMeshToMeshDescription.cpp
--- a/Engine/Source/Runtime/MeshConversion/Private/DynamicMeshToMeshDescription.cpp
+++ b/Engine/Source/Runtime/MeshConversion/Private/DynamicMeshToMeshDescription.cpp
@@ -381,6 +381,30 @@ void FDynamicMeshToMeshDescription::Convert(const FDynamicMesh3* MeshIn, FMeshDe
 	}
 }
 
+
+bool FDynamicMeshToMeshDescription::HaveMatchingElementCounts(const FDynamicMesh3* DynamicMesh, const FMeshDescription* MeshDescription, bool bVerticesOnly, bool bAttributesOnly)
+{
+	bool bVerticesMatch = DynamicMesh->IsCompactV() && DynamicMesh->VertexCount() == MeshDescription->Vertices().Num();
+	bool bTrianglesMatch = DynamicMesh->IsCompactT() && DynamicMesh->TriangleCount() == MeshDescription->Triangles().Num();
+	if (bVerticesOnly || (bAttributesOnly && !DynamicMesh->HasAttributes()))
+	{
+		return bVerticesMatch;
+	}
+	else if (bAttributesOnly && DynamicMesh->HasAttributes())
+	{
+		return bTrianglesMatch;
+	}
+	return bVerticesMatch && bTrianglesMatch;
+}
+
+
+bool FDynamicMeshToMeshDescription::HaveMatchingElementCounts(const FDynamicMesh3* DynamicMesh, const FMeshDescription* MeshDescription)
+{
+	bool bUpdateAttributes = ConversionOptions.bUpdateNormals || ConversionOptions.bUpdateUVs;
+	return HaveMatchingElementCounts(DynamicMesh, MeshDescription, !bUpdateAttributes, !ConversionOptions.bUpdatePositions);
+}
+
+
 void FDynamicMeshToMeshDescription::UpdateUsingConversionOptions(const FDynamicMesh3* MeshIn, FMeshDescription& MeshOut)
 {
 	// See if we can do a buffer update without having to alter triangles.
diff --git a/Engine/Source/Runtime/MeshConversion/Public/DynamicMeshToMeshDescription.h b/Engine/Source/Runtime/MeshConversion/Public/DynamicMeshToMeshDescription.h
--- a/Engine/Source/Runtime/MeshConversion/Public/DynamicMeshToMeshDescription.h
+++ b/Engine/Source/Runtime/MeshConversion/Public/DynamicMeshToMeshDescription.h
@@ -40,20 +40,7 @@ public:
 	 * @param bAttributesOnly If true, only check what needs to be checked for UpdateAttributes
 							 (will check vertices or triangles depending on whether attributes are per vertex or in overlays)
 	 */
-	static bool HaveMatchingElementCounts(const FDynamicMesh3* DynamicMesh, const FMeshDescription* MeshDescription, bool bVerticesOnly, bool bAttributesOnly)
-	{
-		bool bVerticesMatch = DynamicMesh->IsCompactV() && DynamicMesh->VertexCount() == MeshDescription->Vertices().Num();
-		bool bTrianglesMatch = DynamicMesh->IsCompactT() && DynamicMesh->TriangleCount() == MeshDescription->Triangles().Num();
-		if (bVerticesOnly || (bAttributesOnly && !DynamicMesh->HasAttributes()))
-		{
-			return bVerticesMatch;
-		}
-		else if (bAttributesOnly && DynamicMesh->HasAttributes())
-		{
-			return bTrianglesMatch;
-		}
-		return bVerticesMatch && bTrianglesMatch;
-	}
+	static bool HaveMatchingElementCounts(const FDynamicMesh3* DynamicMesh, const FMeshDescription* MeshDescription, bool bVerticesOnly, bool bAttributesOnly);
 
 	/**
 	 * Checks if element counts match. If false then Update can't be called -- you must call Convert
@@ -62,11 +49,7 @@ public:
 	 * @param DynamicMesh The dynamic mesh with updated vertices or attributes
 	 * @param MeshDescription The corresponding mesh description
 	 */
-	bool HaveMatchingElementCounts(const FDynamicMesh3* DynamicMesh, const FMeshDescription* MeshDescription)
-	{
-		bool bUpdateAttributes = ConversionOptions.bUpdateNormals || ConversionOptions.bUpdateUVs;
-		return HaveMatchingElementCounts(DynamicMesh, MeshDescription, !bUpdateAttributes, !ConversionOptions.bUpdatePositions);
-	}
+	bool HaveMatchingElementCounts(const FDynamicMesh3* DynamicMesh, const FMeshDescription* MeshDescription);
 
 	/**
 	 * Default conversion of DynamicMesh to MeshDescription. Calls functions below depending on mesh state
