Fix an issue with editing skelmeshes without material slot names. All materials would use the first material slot as a fallback after editing was committed. Also fixed a related issue where the default profile skinweights attribute wasn't being registered on the resulting meshdesc conversion from dynmesh.

diff --git a/Engine/Source/Runtime/Engine/Private/SkeletalMesh.cpp b/Engine/Source/Runtime/Engine/Private/SkeletalMesh.cpp
--- a/Engine/Source/Runtime/Engine/Private/SkeletalMesh.cpp
+++ b/Engine/Source/Runtime/Engine/Private/SkeletalMesh.cpp
@@ -3573,6 +3573,14 @@ void USkeletalMesh::ExecutePostLoadInternal(FSkeletalMeshPostLoadContext& Contex
 
 		PostLoadVerifyAndFixBadTangent();
 
+		// Fixup missing material slot names and import slot names, so that mesh editing 
+		// preserves material assignments.
+		IMeshUtilities* MeshUtilities = FModuleManager::Get().LoadModulePtr<IMeshUtilities>("MeshUtilities");
+		if (MeshUtilities)
+		{
+			MeshUtilities->FixupMaterialSlotNames(this);
+		}
+
 		if (GetResourceForRendering() == nullptr)
 		{
 			CacheDerivedData(&Context);
diff --git a/Engine/Source/Runtime/Engine/Private/SkeletalMeshLODModel.cpp b/Engine/Source/Runtime/Engine/Private/SkeletalMeshLODModel.cpp
--- a/Engine/Source/Runtime/Engine/Private/SkeletalMeshLODModel.cpp
+++ b/Engine/Source/Runtime/Engine/Private/SkeletalMeshLODModel.cpp
@@ -1435,12 +1435,12 @@ void FSkeletalMeshLODModel::GetMeshDescription(FMeshDescription& MeshDescription
 
 		if (!MeshDescription.IsPolygonGroupValid(PolygonGroupID))
 		{
-			MeshDescription.CreatePolygonGroupWithID(FPolygonGroupID(Section.MaterialIndex));
+			MeshDescription.CreatePolygonGroupWithID(PolygonGroupID);
 		}
 
 		if (ensure(Materials.IsValidIndex(Section.MaterialIndex)))
 		{
-			PolygonGroupMaterialSlotNames.Set(PolygonGroupID, Materials[Section.MaterialIndex].MaterialSlotName);
+			PolygonGroupMaterialSlotNames.Set(PolygonGroupID, Materials[Section.MaterialIndex].ImportedMaterialSlotName);
 		}
 
 		for (int32 TriangleID = 0; TriangleID < int32(Section.NumTriangles); TriangleID++)
diff --git a/Engine/Source/Runtime/MeshConversion/Private/DynamicMeshToMeshDescription.cpp b/Engine/Source/Runtime/MeshConversion/Private/DynamicMeshToMeshDescription.cpp
--- a/Engine/Source/Runtime/MeshConversion/Private/DynamicMeshToMeshDescription.cpp
+++ b/Engine/Source/Runtime/MeshConversion/Private/DynamicMeshToMeshDescription.cpp
@@ -674,7 +674,6 @@ void FDynamicMeshToMeshDescription::Convert_SharedInstances(const FDynamicMesh3*
 
 
 
-
 void FDynamicMeshToMeshDescription::Convert_NoSharedInstances(const FDynamicMesh3* MeshIn, FMeshDescription& MeshOut, bool bCopyTangents)
 {
 	
@@ -711,6 +710,20 @@ void FDynamicMeshToMeshDescription::Convert_NoSharedInstances(const FDynamicMesh
 		bCopyGroupToPolyGroup = true;
 	}
 
+	TMap<FName, FSkinWeightsVertexAttributesRef> VertexBoneWeightsMap; 
+	if (!MeshIn->Attributes()->GetSkinWeightsAttributes().IsEmpty())
+	{
+		FSkeletalMeshAttributes MeshOutAttributes(MeshOut);
+		MeshOutAttributes.Register();
+		for (const TTuple<FName, TUniquePtr<FDynamicMeshVertexSkinWeightsAttribute>>& AttributeInfo: MeshIn->Attributes()->GetSkinWeightsAttributes())
+		{
+			FName ProfileName = AttributeInfo.Key;
+
+			MeshOutAttributes.RegisterSkinWeightAttribute(ProfileName);
+			VertexBoneWeightsMap.Add(ProfileName, MeshOutAttributes.GetVertexSkinWeights(ProfileName));
+		}
+	}
+
 	// always copy when we are baking new mesh? should this be a config option?
 	bool bCopyInstanceColors = (ColorOverlay != nullptr); 
 
@@ -746,7 +759,6 @@ void FDynamicMeshToMeshDescription::Convert_NoSharedInstances(const FDynamicMesh
 			MapUV[ElementID] = Builder.AppendUV(UVvalue, k);
 		}
 	}
-	
 
 
 	FPolygonGroupID ZeroPolygonGroupID = Builder.AppendPolygonGroup();
@@ -939,17 +951,13 @@ void FDynamicMeshToMeshDescription::Convert_NoSharedInstances(const FDynamicMesh
 
 	// Convert all attached skin weights, if we're converting a mesh description that originated from
 	// a USkeletalMesh.
-	FSkeletalMeshAttributes MeshOutAttributes(MeshOut);
 	for (const TTuple<FName, TUniquePtr<FDynamicMeshVertexSkinWeightsAttribute>>& AttributeInfo: MeshIn->Attributes()->GetSkinWeightsAttributes())
 	{
 		FName ProfileName = AttributeInfo.Key;
-
-		MeshOutAttributes.RegisterSkinWeightAttribute(ProfileName);
-
-		const FDynamicMeshVertexSkinWeightsAttribute *MeshSkinWeights = AttributeInfo.Value.Get();
 		
-		FSkinWeightsVertexAttributesRef VertexBoneWeights = MeshOutAttributes.GetVertexSkinWeights(ProfileName);
+		FSkinWeightsVertexAttributesRef& VertexBoneWeights = VertexBoneWeightsMap[ProfileName];
 
+		const FDynamicMeshVertexSkinWeightsAttribute *MeshSkinWeights = AttributeInfo.Value.Get();
 		for (int32 VertexIndex = 0; VertexIndex < MapV.Num(); VertexIndex++)
 		{
 			if (const FVertexID VertexID = MapV[VertexIndex]; VertexID != INDEX_NONE)
@@ -965,7 +973,6 @@ void FDynamicMeshToMeshDescription::Convert_NoSharedInstances(const FDynamicMesh
 }
 
 
-
 void FDynamicMeshToMeshDescription::ConvertPolygroupLayers(const FDynamicMesh3* MeshIn, FMeshDescription& MeshOut, const TArray<FTriangleID>& IndexToTriangleIDMap)
 {
 	if (MeshIn->Attributes() == nullptr) return;
