StaticMeshLODResourcesToDynamicMesh.cpp - build fix

diff --git a/Engine/Source/Runtime/MeshConversion/Public/StaticMeshLODResourcesAdapter.h b/Engine/Source/Runtime/MeshConversion/Public/StaticMeshLODResourcesAdapter.h
--- a/Engine/Source/Runtime/MeshConversion/Public/StaticMeshLODResourcesAdapter.h
+++ b/Engine/Source/Runtime/MeshConversion/Public/StaticMeshLODResourcesAdapter.h
@@ -18,6 +18,10 @@ struct FStaticMeshLODResourcesMeshAdapter
 protected:
 	const FStaticMeshLODResources* Mesh;
 
+	FVector3d BuildScale = FVector3d::One();
+	FVector3d InvBuildScale = FVector3d::One();
+	bool bScaleNormals = false;
+
 	TArray<const FStaticMeshSection*> ValidSections;
 	TArray<int32> TriangleOffsetArray;
 	int32 NumTriangles;
@@ -47,6 +51,13 @@ public:
 		TriangleOffsetArray.Add(NumTriangles);
 	}
 
+	void SetBuildScale(const FVector3d& BuildScaleIn, bool bScaleNormalsIn)
+	{
+		BuildScale = BuildScaleIn;
+		InvBuildScale = UE::Geometry::FTransformSRT3d::GetSafeScaleReciprocal(BuildScaleIn);
+		bScaleNormals = bScaleNormalsIn;
+	}
+
 	bool IsTriangle(int32 TID) const
 	{
 		return TID >= 0 && TID < MaxTriangleID();
@@ -98,7 +109,7 @@ public:
 	inline FVector3d GetVertex(int32 IDValue) const
 	{
 		const FVector3f& Position = Mesh->VertexBuffers.PositionVertexBuffer.VertexPosition(IDValue);
-		return FVector3d(Position.X, Position.Y, Position.Z);
+		return FVector3d(BuildScale.X * (double)Position.X, BuildScale.Y * (double)Position.Y, BuildScale.Z * (double)Position.Z);
 	}
 
 	inline void GetTriVertices(int32 IDValue, FVector3d& V0, FVector3d& V1, FVector3d& V2) const
@@ -140,7 +151,8 @@ public:
 	inline FVector3f GetNormal(int32 IDValue) const
 	{
 		const FVector4f& Normal = Mesh->VertexBuffers.StaticMeshVertexBuffer.VertexTangentZ(IDValue);
-		return FVector3f(Normal.X, Normal.Y, Normal.Z);
+		return (!bScaleNormals) ? FVector3f(Normal.X, Normal.Y, Normal.Z) :
+			UE::Geometry::Normalized(FVector3f(Normal.X * InvBuildScale.X, Normal.Y * InvBuildScale.Y, Normal.Z * InvBuildScale.Z));
 	}
 
 	/** Get Normals for a given Triangle */
@@ -158,7 +170,8 @@ public:
 	inline FVector3f GetTangentX(int32 IDValue) const
 	{
 		const FVector4f& TangentX = Mesh->VertexBuffers.StaticMeshVertexBuffer.VertexTangentX(IDValue);
-		return FVector3f(TangentX.X, TangentX.Y, TangentX.Z);
+		return (!bScaleNormals) ? FVector3f(TangentX.X, TangentX.Y, TangentX.Z) :
+			UE::Geometry::Normalized(FVector3f(TangentX.X * BuildScale.X, TangentX.Y * BuildScale.Y, TangentX.Z * BuildScale.Z));
 	}
 
 	/** Get Tangent X for a given Triangle */
@@ -175,7 +188,8 @@ public:
 	inline FVector3f GetTangentY(int32 IDValue) const
 	{
 		const FVector4f& TangentY = Mesh->VertexBuffers.StaticMeshVertexBuffer.VertexTangentY(IDValue);
-		return FVector3f(TangentY.X, TangentY.Y, TangentY.Z);
+		return (!bScaleNormals) ? FVector3f(TangentY.X, TangentY.Y, TangentY.Z) :
+			UE::Geometry::Normalized(FVector3f(TangentY.X * BuildScale.X, TangentY.Y * BuildScale.Y, TangentY.Z * BuildScale.Z));
 	}
 
 	/** Get Tangent Y for a given Triangle */
@@ -260,4 +274,4 @@ public:
 		C2 = GetColor(VtxIndices.C);
 	}
 
-};
+};
\ No newline at end of file
