MeshConversion: Resolve static analysis warning

diff --git a/Engine/Source/Runtime/MeshConversion/Private/DynamicMeshToMeshDescription.cpp b/Engine/Source/Runtime/MeshConversion/Private/DynamicMeshToMeshDescription.cpp
--- a/Engine/Source/Runtime/MeshConversion/Private/DynamicMeshToMeshDescription.cpp
+++ b/Engine/Source/Runtime/MeshConversion/Private/DynamicMeshToMeshDescription.cpp
@@ -263,20 +263,22 @@ void FDynamicMeshToMeshDescription::UpdateTangents(const FDynamicMesh3* MeshIn,
 	if (!ensureMsgf(MeshIn->TriangleCount() == MeshOut.Triangles().Num(), TEXT("Trying to update MeshDescription Tangents from Mesh that does not have same triangle count"))) return;
 	if (!ensureMsgf(MeshIn->HasAttributes(), TEXT("Trying to update MeshDescription Tangents from a DynamicMesh that has no attributes, e.g. normals"))) return;
 
-	
 	// src
 	const FDynamicMeshNormalOverlay* NormalOverlay = MeshIn->Attributes()->PrimaryNormals();
 	const FDynamicMeshNormalOverlay* TangentOverlay = MeshIn->Attributes()->PrimaryTangents();
 	const FDynamicMeshNormalOverlay* BiTangentOverlay = MeshIn->Attributes()->PrimaryBiTangents();
 
-	bool bHasValidSrc = (NormalOverlay != nullptr) && (TangentOverlay != nullptr) && (BiTangentOverlay != nullptr);
-
-	if(!ensureMsgf(bHasValidSrc, TEXT("Trying to update MeshDescription Tangents from a DynamicMesh that does not have all three tangent space attributes"))) return;
+	const bool bHasValidSrc = NormalOverlay && TangentOverlay && BiTangentOverlay;
+	ensureMsgf(bHasValidSrc, TEXT("Trying to update MeshDescription Tangents from a DynamicMesh that does not have all three tangent space attributes"));
+	if (!bHasValidSrc)
+	{
+		return;
+	}
 
 	// dst
 	FStaticMeshAttributes Attributes(MeshOut);
-	TVertexInstanceAttributesRef<FVector3f> TangentAttrib = Attributes.GetVertexInstanceTangents();
-	TVertexInstanceAttributesRef<float> BiTangentSignAttrib = Attributes.GetVertexInstanceBinormalSigns();
+	const TVertexInstanceAttributesRef<FVector3f> TangentAttrib = Attributes.GetVertexInstanceTangents();
+	const TVertexInstanceAttributesRef<float> BiTangentSignAttrib = Attributes.GetVertexInstanceBinormalSigns();
 
 	if (!ensureMsgf(TangentAttrib.IsValid(), TEXT("Trying to update Tangents on a MeshDescription that has no Tangent Vertex Instance attribute"))) return;
 	if (!ensureMsgf(BiTangentSignAttrib.IsValid(), TEXT("Trying to update Tangents on a MeshDescription that has no BinormalSign Vertex Instance attribute"))) return;
