MeshConversion: Add flag to undo pre-transformation of vertex colors during MeshDescriptionToDynamicMesh.
Make pre-transformation of vertex colors during MeshDescription/DynamicMesh conversion on by default.
Fixed Convert_NoSharedInstances code path that was not pre-transforming vertex colors.

diff --git a/Engine/Source/Runtime/MeshConversion/Private/DynamicMeshToMeshDescription.cpp b/Engine/Source/Runtime/MeshConversion/Private/DynamicMeshToMeshDescription.cpp
--- a/Engine/Source/Runtime/MeshConversion/Private/DynamicMeshToMeshDescription.cpp
+++ b/Engine/Source/Runtime/MeshConversion/Private/DynamicMeshToMeshDescription.cpp
@@ -339,24 +339,9 @@ void FDynamicMeshToMeshDescription::UpdateVertexColors(const FDynamicMesh3* Mesh
 		FVector4f TriColors[3];
 		ColorOverlay->GetTriElements(t, TriColors[0], TriColors[1], TriColors[2]);
 
-		// There is inconsistency in how vertex colors are intended to be consumed in
-		// our shaders. Some shaders consume it as linear (ex. MeshPaint), others as SRGB which
-		// manually convert to linear in the shader.
-		//
-		// All StaticMeshes store vertex colors as an 8-bit FColor. In order to ensure a good
-		// distribution of float values across the 8-bit range, the StaticMesh build always
-		// encodes FColors as SRGB.
-		//
-		// Until there is some defined gamma space convention for vertex colors in our shaders,
-		// we provide this option to pre-transform our linear float colors with an SRGBToLinear
-		// conversion to counteract the StaticMesh build LinearToSRGB conversion. This is how
-		// MeshPaint ensures linear vertex colors in the shaders.
-		if (ConversionOptions.bTransformVtxColorsSRGBToLinear)
-		{
-			LinearColors::SRGBToLinear(TriColors[0]);
-			LinearColors::SRGBToLinear(TriColors[1]);
-			LinearColors::SRGBToLinear(TriColors[2]);
-		}
+		ApplyVertexColorTransform(TriColors[0]);
+		ApplyVertexColorTransform(TriColors[1]);
+		ApplyVertexColorTransform(TriColors[2]);
 		
 		TArrayView<const FVertexInstanceID> InstanceIDs = MeshOut.GetTriangleVertexInstances(FTriangleID(t));
 		for (int32 i = 0; i < 3; ++i)
@@ -818,7 +803,7 @@ void FDynamicMeshToMeshDescription::Convert_NoSharedInstances(const FDynamicMesh
 	TFunction<void(int TriID, FVertexInstanceID TriVertInstances[3])> ColorInstanceSetter;
 	if (bCopyInstanceColors)
 	{
-		ColorInstanceSetter = [ColorOverlay, &Builder](int TriID, FVertexInstanceID TriVertInstances[3])
+		ColorInstanceSetter = [this, ColorOverlay, &Builder](int TriID, FVertexInstanceID TriVertInstances[3])
 		{
 			FIndex3i ColorTri = ColorOverlay->GetTriangle(TriID);
 			for (int32 j = 0; j < 3; ++j)
@@ -828,6 +813,7 @@ void FDynamicMeshToMeshDescription::Convert_NoSharedInstances(const FDynamicMesh
 				if (ColorOverlay->IsElement(ColorTri[j]))
 				{
 					FVector4f TriVertColor4 = ColorOverlay->GetElement(ColorTri[j]);
+					ApplyVertexColorTransform(TriVertColor4);
 					DstColor = FVector4f(TriVertColor4.X, TriVertColor4.Y, TriVertColor4.Z, TriVertColor4.W);
 				}
 				Builder.SetInstanceColor(CornerInstanceID, DstColor);
@@ -1072,3 +1058,24 @@ void FDynamicMeshToMeshDescription::ConvertWeightLayers(const FDynamicMesh3* Mes
 		}
 	}
 }
+
+void FDynamicMeshToMeshDescription::ApplyVertexColorTransform(FVector4f& Color) const
+{
+	// There is inconsistency in how vertex colors are intended to be consumed in
+	// our shaders. Some shaders consume it as linear (ex. MeshPaint), others as SRGB which
+	// manually convert to linear in the shader.
+	//
+	// All StaticMeshes store vertex colors as an 8-bit FColor. In order to ensure a good
+	// distribution of float values across the 8-bit range, the StaticMesh build always
+	// encodes FColors as SRGB.
+	//
+	// Until there is some defined gamma space convention for vertex colors in our shaders,
+	// we provide this option to pre-transform our linear float colors with an SRGBToLinear
+	// conversion to counteract the StaticMesh build LinearToSRGB conversion. This is how
+	// MeshPaint ensures linear vertex colors in the shaders.
+	if (ConversionOptions.bTransformVtxColorsSRGBToLinear)
+	{
+		LinearColors::SRGBToLinear(Color);
+	}
+}
+
diff --git a/Engine/Source/Runtime/MeshConversion/Private/MeshDescriptionToDynamicMesh.cpp b/Engine/Source/Runtime/MeshConversion/Private/MeshDescriptionToDynamicMesh.cpp
--- a/Engine/Source/Runtime/MeshConversion/Private/MeshDescriptionToDynamicMesh.cpp
+++ b/Engine/Source/Runtime/MeshConversion/Private/MeshDescriptionToDynamicMesh.cpp
@@ -9,6 +9,7 @@
 #include "MeshDescriptionBuilder.h"
 #include "StaticMeshAttributes.h"
 #include "SkeletalMeshAttributes.h"
+#include "Util/ColorConstants.h"
 #include "Async/Async.h"
 
 using namespace UE::Geometry;
@@ -761,7 +762,9 @@ void FMeshDescriptionToDynamicMesh::Convert(const FMeshDescription* MeshIn, FDyn
 					FIndex3i TriVector;
 					for (int j = 0; j < 3; ++j)
 					{
-						const FVector4f InstanceColor4 = InstanceColors.Get(TriData.TriInstances[j], 0);
+						FVector4f InstanceColor4 = InstanceColors.Get(TriData.TriInstances[j], 0);
+						ApplyVertexColorTransform(InstanceColor4);
+						
 						const FVector4f OverlayColor(InstanceColor4.X, InstanceColor4.Y, InstanceColor4.Z, InstanceColor4.W);
 						TriVector[j] = ColorWelder.FindOrAddUnique(OverlayColor, Tri[j]);
 
@@ -930,3 +933,23 @@ void FMeshDescriptionToDynamicMesh::CopyTangents(const FMeshDescription* SourceM
 	if (!ensureMsgf(TriIDMap.Num() == TargetMesh->TriangleCount(), TEXT("Tried to CopyTangents to mesh with different triangle count"))) return;
 	CopyTangents_Internal<double>(SourceMesh, TargetMesh, TangentsOut, TriIDMap);
 }
+
+void FMeshDescriptionToDynamicMesh::ApplyVertexColorTransform(FVector4f& Color) const
+{
+	if (bTransformVertexColorsLinearToSRGB)
+	{
+		// The corollary to bTransformVertexColorsSRGBToLinear in DynamicMeshToMeshDescription.
+		// See DynamicMeshToMeshDescription::ApplyVertexColorTransform(..).
+		//
+		// StaticMeshes store vertex colors as FColor. The StaticMesh build always encodes
+		// FColors as SRGB to ensure a good distribution of float values across the 8-bit range.
+		//
+		// Since there is no currently defined gamma space convention for vertex colors in
+		// engine, an option is provided to pre-transform vertex colors (SRGB To Linear) when
+		// writing out a MeshDescription.
+		//
+		// We similarly provide the inverse of that optional pre-transformation to maintain
+		// color space consistency in our usage.
+		LinearColors::LinearToSRGB(Color);
+	}
+}
diff --git a/Engine/Source/Runtime/MeshConversion/Public/DynamicMeshToMeshDescription.h b/Engine/Source/Runtime/MeshConversion/Public/DynamicMeshToMeshDescription.h
--- a/Engine/Source/Runtime/MeshConversion/Public/DynamicMeshToMeshDescription.h
+++ b/Engine/Source/Runtime/MeshConversion/Public/DynamicMeshToMeshDescription.h
@@ -172,4 +172,14 @@ protected:
 	*/
 	void ConvertWeightLayers(const FDynamicMesh3* MeshIn, FMeshDescription& MeshOut, const TArray<FVertexID>& IndexToVertexIDMap);
 
+	/**
+	 * Applies an optional sRGB-to-Linear color transform on the input. The color transform
+	 * is controlled by ConversionOptions.bTransformVtxColorsSRGBToLinear.
+	 *
+	 * The counterpart to this method is MeshDescriptionToDynamicMesh::ApplyVertexColorTransform
+	 * which will undo this color transformation when the MeshDescription is read back.
+	 *
+	 * @param Color color to transform
+	 */
+	void ApplyVertexColorTransform(FVector4f& Color) const;
 };
\ No newline at end of file
diff --git a/Engine/Source/Runtime/MeshConversion/Public/MeshConversionOptions.h b/Engine/Source/Runtime/MeshConversion/Public/MeshConversionOptions.h
--- a/Engine/Source/Runtime/MeshConversion/Public/MeshConversionOptions.h
+++ b/Engine/Source/Runtime/MeshConversion/Public/MeshConversionOptions.h
@@ -28,7 +28,7 @@ public:
 	bool bUpdateVtxColors = false;
 
 	/** Should Vertex Colors of MeshDescription be transformed from SRGB to Linear */
-	bool bTransformVtxColorsSRGBToLinear = false;
+	bool bTransformVtxColorsSRGBToLinear = true;
 
 	//
 	// utility functions for common configuration cases
diff --git a/Engine/Source/Runtime/MeshConversion/Public/MeshDescriptionToDynamicMesh.h b/Engine/Source/Runtime/MeshConversion/Public/MeshDescriptionToDynamicMesh.h
--- a/Engine/Source/Runtime/MeshConversion/Public/MeshDescriptionToDynamicMesh.h
+++ b/Engine/Source/Runtime/MeshConversion/Public/MeshDescriptionToDynamicMesh.h
@@ -39,6 +39,8 @@ public:
 	/** Ignore all mesh attributes (e.g. UV/Normal layers, color layer, material groups) */
 	bool bDisableAttributes = false;
 
+	/** Should Vertex Colors of MeshDescription be transformed from Linear to SRGB */
+	bool bTransformVertexColorsLinearToSRGB = true;
 
 	/** map from DynamicMesh triangle ID to MeshDescription FTriangleID*/
 	TArray<FTriangleID> TriIDMap;
@@ -85,4 +87,16 @@ public:
 	 * @warning Convert() must have been used to create the TargetMesh before calling this function
 	 */
 	void CopyTangents(const FMeshDescription* SourceMesh, const FDynamicMesh3* TargetMesh, UE::Geometry::TMeshTangents<double>* TangentsOut);
+
+protected:
+	/**
+	 * Applies an optional Linear-to-sRGB color transform on the input. The color transform
+	 * is controlled by bTransformVtxColorsLinearToSRGB.
+	 *
+	 * This is the counterpart to DynamicMeshToMeshDescription::ApplyVertexColorTransform to
+	 * undo an applied sRGB-to-Linear transformation when the MeshDescription was built.
+	 *
+	 * @param Color color to transform
+	 */
+	void ApplyVertexColorTransform(FVector4f& Color) const;
 };
\ No newline at end of file
