The problem appears when a process of a set of StaticMesh calling ConverterToDynamicMesh is parallelized.
TFuture<void> / Async is not adapted for parallelization of parallelized process.
On the advice of Andriy Tylychko, TFuture has been replaced by UE::Tasks::FTask
#swarm https://p4-swarm.epicgames.net/reviews/21225362

diff --git a/Engine/Plugins/Enterprise/DatasmithCADImporter/Source/DatasmithCADTranslator/Private/DatasmithCADTranslator.cpp b/Engine/Plugins/Enterprise/DatasmithCADImporter/Source/DatasmithCADTranslator/Private/DatasmithCADTranslator.cpp
--- a/Engine/Plugins/Enterprise/DatasmithCADImporter/Source/DatasmithCADTranslator/Private/DatasmithCADTranslator.cpp
+++ b/Engine/Plugins/Enterprise/DatasmithCADImporter/Source/DatasmithCADTranslator/Private/DatasmithCADTranslator.cpp
@@ -112,9 +112,6 @@ void FDatasmithCADTranslator::Initialize(FDatasmithTranslatorCapabilities& OutCa
 
 	OutCapabilities.SupportedFileFormats.Add(FFileFormatInfo{ TEXT("hsf"), TEXT("HOOPS stream files") });
 	OutCapabilities.SupportedFileFormats.Add(FFileFormatInfo{ TEXT("prc"), TEXT("HOOPS stream files") });
-
-	// FDatasmithCADTranslator::LoadStaticMesh is not compilant with parallel loading if MeshOperator::ResolveTJunctions is used
-	OutCapabilities.bParallelLoadStaticMeshSupported = !CADLibrary::FImportParameters::bGSewMeshIfNeeded;
 }
 
 bool FDatasmithCADTranslator::IsSourceSupported(const FDatasmithSceneSource& Source)
diff --git a/Engine/Plugins/Enterprise/DatasmithImporter/Source/DatasmithImporter/Private/DatasmithImporter.cpp b/Engine/Plugins/Enterprise/DatasmithImporter/Source/DatasmithImporter/Private/DatasmithImporter.cpp
--- a/Engine/Plugins/Enterprise/DatasmithImporter/Source/DatasmithImporter/Private/DatasmithImporter.cpp
+++ b/Engine/Plugins/Enterprise/DatasmithImporter/Source/DatasmithImporter/Private/DatasmithImporter.cpp
@@ -35,8 +35,6 @@
 
 #include "AssetRegistry/AssetRegistryModule.h"
 #include "AssetToolsModule.h"
-#include "Async/Async.h"
-#include "Async/AsyncWork.h"
 #include "CineCameraComponent.h"
 #include "ComponentReregisterContext.h"
 #include "Components/HierarchicalInstancedStaticMeshComponent.h"
@@ -75,6 +73,7 @@
 #include "Settings/EditorExperimentalSettings.h"
 #include "SkinnedAssetCompiler.h"
 #include "SourceControlOperations.h"
+#include "Tasks/Task.h"
 #include "Templates/UniquePtr.h"
 #include "UObject/Package.h"
 #include "UObject/UObjectHash.h"
@@ -264,7 +263,7 @@ void FDatasmithImporter::ImportStaticMeshes( FDatasmithImportContext& ImportCont
 		ProgressPtr->MakeDialog(true);
 	}
 
-	TMap<TSharedRef<IDatasmithMeshElement>, TFuture<FDatasmithMeshElementPayload*>> MeshElementPayloads;
+	TMap<TSharedRef<IDatasmithMeshElement>, UE::Tasks::TTask<FDatasmithMeshElementPayload*>> MeshElementPayloads;
 
 	FDatasmithTranslatorCapabilities TranslatorCapabilities;
 	ImportContext.SceneTranslator->Initialize(TranslatorCapabilities);
@@ -289,8 +288,7 @@ void FDatasmithImporter::ImportStaticMeshes( FDatasmithImportContext& ImportCont
 				// Parallel loading from the translator using futures
 				MeshElementPayloads.Add(
 					MeshElement,
-					Async(
-						EAsyncExecution::LargeThreadPool,
+					UE::Tasks::Launch(UE_SOURCE_LOCATION,
 						[&ImportContext, MeshElement]() -> FDatasmithMeshElementPayload*
 						{
 							if (ImportContext.bUserCancelled)
@@ -338,11 +336,11 @@ void FDatasmithImporter::ImportStaticMeshes( FDatasmithImportContext& ImportCont
 		//  - GetDestination (find or create StaticMesh, duplicate, flags and context etc)
 		//  - Import (Import data in simple memory repr (eg. TArray<FMeshDescription>)
 		//  - Set (fill UStaticMesh with imported data)
-		TFuture<FDatasmithMeshElementPayload*> MeshPayload;
+		UE::Tasks::TTask<FDatasmithMeshElementPayload*> MeshPayload;
 		if (MeshElementPayloads.RemoveAndCopyValue(MeshElement, MeshPayload))
 		{
-			TUniquePtr<FDatasmithMeshElementPayload> MeshPayloadPtr(MeshPayload.Get());
-			if (MeshPayloadPtr.IsValid())
+			TUniquePtr<FDatasmithMeshElementPayload> MeshPayloadPtr(MeshPayload.GetResult());
+			if (MeshPayloadPtr)
 			{
 				ImportStaticMesh(ImportContext, MeshElement, ExistingStaticMesh, MeshPayloadPtr.Get());
 			}
@@ -356,10 +354,10 @@ void FDatasmithImporter::ImportStaticMeshes( FDatasmithImportContext& ImportCont
 	}
 
 	//Just make sure there is no async task left running in case of a cancellation
-	for ( const TPair<TSharedRef<IDatasmithMeshElement>, TFuture<FDatasmithMeshElementPayload*>> & Kvp : MeshElementPayloads)
+	for (TPair<TSharedRef<IDatasmithMeshElement>, UE::Tasks::TTask<FDatasmithMeshElementPayload*>>& Kvp : MeshElementPayloads)
 	{
 		// Wait for the result and delete it when getting out of scope
-		TUniquePtr<FDatasmithMeshElementPayload> MeshPayloadPtr(Kvp.Value.Get());
+		TUniquePtr<FDatasmithMeshElementPayload> MeshPayloadPtr(Kvp.Value.GetResult());
 	}
 
 	TMap< TSharedRef< IDatasmithMeshElement >, float > LightmapWeights = FDatasmithStaticMeshImporter::CalculateMeshesLightmapWeights( ImportContext.Scene.ToSharedRef() );
@@ -695,9 +693,10 @@ void FDatasmithImporter::ImportTextures( FDatasmithImportContext& ImportContext
 		{
 			FString       Extension;
 			TArray<uint8> TextureData;
-			TFuture<bool> Result;
+			UE::Tasks::TTask<bool> Result;
 		};
 		TArray<FAsyncData> AsyncData;
+
 		AsyncData.SetNum(FilteredTextureElements.Num());
 
 		for ( int32 TextureIndex = 0; TextureIndex < FilteredTextureElements.Num(); TextureIndex++ )
@@ -713,8 +712,7 @@ void FDatasmithImporter::ImportTextures( FDatasmithImportContext& ImportContext
 			}
 
 			AsyncData[TextureIndex].Result =
-				Async(
-					EAsyncExecution::LargeThreadPool,
+				UE::Tasks::Launch(UE_SOURCE_LOCATION,
 					[&ImportContext, &AsyncData, &FilteredTextureElements, &DatasmithTextureImporter, TextureIndex]()
 					{
 						if (ImportContext.bUserCancelled)
@@ -770,7 +768,7 @@ void FDatasmithImporter::ImportTextures( FDatasmithImportContext& ImportContext
 					}
 				}
 
-				if (AsyncData[TextureIndex].Result.Get())
+				if (AsyncData[TextureIndex].Result.GetResult())
 				{
 					ImportTexture( ImportContext, DatasmithTextureImporter, TextureElement.ToSharedRef(), ExistingTexture, AsyncData[TextureIndex].TextureData, AsyncData[TextureIndex].Extension );
 				}
diff --git a/Engine/Source/Runtime/MeshConversion/Private/MeshDescriptionToDynamicMesh.cpp b/Engine/Source/Runtime/MeshConversion/Private/MeshDescriptionToDynamicMesh.cpp
--- a/Engine/Source/Runtime/MeshConversion/Private/MeshDescriptionToDynamicMesh.cpp
+++ b/Engine/Source/Runtime/MeshConversion/Private/MeshDescriptionToDynamicMesh.cpp
@@ -1,6 +1,7 @@
 // Copyright Epic Games, Inc. All Rights Reserved. 
 
 #include "MeshDescriptionToDynamicMesh.h"
+
 #include "DynamicMesh/DynamicMeshAttributeSet.h"
 #include "DynamicMesh/DynamicMeshOverlay.h"
 #include "DynamicMesh/DynamicVertexAttribute.h"
@@ -9,8 +10,8 @@
 #include "MeshDescriptionBuilder.h"
 #include "StaticMeshAttributes.h"
 #include "SkeletalMeshAttributes.h"
+#include "Tasks/Task.h"
 #include "Util/ColorConstants.h"
-#include "Async/Async.h"
 
 using namespace UE::Geometry;
 
@@ -171,6 +172,8 @@ struct FSkinWeightsAttribCopyInfo
 
 void FMeshDescriptionToDynamicMesh::Convert(const FMeshDescription* MeshIn, FDynamicMesh3& MeshOut, bool bCopyTangents )
 {
+	TRACE_CPUPROFILER_EVENT_SCOPE(MeshDescriptionToDynamicMesh);
+
 	TriIDMap.Reset();
 	VertIDMap.Reset();
 
@@ -494,7 +497,7 @@ void FMeshDescriptionToDynamicMesh::Convert(const FMeshDescription* MeshIn, FDyn
 		bool bFoundNonDefaultVertexInstanceColor = false;
 
 		// we will weld/populate all the attributes simultaneously, hold on to futures in this array and then Wait for them at the end
-		TArray<TFuture<void>> Pending;
+		TArray<UE::Tasks::FTask> Pending;
 
 		if (PolygroupAttribs.Num() > 0)
 		{
@@ -505,7 +508,7 @@ void FMeshDescriptionToDynamicMesh::Convert(const FMeshDescription* MeshIn, FDyn
 
 		for (int UVLayerIndex = 0; UVLayerIndex < NumUVLayers; UVLayerIndex++)
 		{
-			auto UVFuture = Async(EAsyncExecution::ThreadPool, [&, UVLayerIndex, bUseSharedUVs]() // must copy UVLayerIndex here!
+			UE::Tasks::FTask UVTask = UE::Tasks::Launch(UE_SOURCE_LOCATION, [&, UVLayerIndex, bUseSharedUVs]() // must copy UVLayerIndex here!
 			{
 				// the overlay to fill.
 				FDynamicMeshUVOverlay* UVOverlay = MeshOut.Attributes()->GetUVLayer(UVLayerIndex);
@@ -669,13 +672,13 @@ void FMeshDescriptionToDynamicMesh::Convert(const FMeshDescription* MeshIn, FDyn
 					}
 				}
 			});
-			Pending.Add(MoveTemp(UVFuture));
+			Pending.Add(MoveTemp(UVTask));
 		}
 
 
 		if (NormalOverlay != nullptr)
 		{
-			auto NormalFuture = Async(EAsyncExecution::ThreadPool, [&]()
+			UE::Tasks::FTask NormalTask = UE::Tasks::Launch(UE_SOURCE_LOCATION, [&]()
 			{
 				FNormalWelder NormalWelder;
 				NormalWelder.NormalOverlay = NormalOverlay;
@@ -692,13 +695,13 @@ void FMeshDescriptionToDynamicMesh::Convert(const FMeshDescription* MeshIn, FDyn
 					NormalOverlay->SetTriangle(TriangleID, TriNormals);
 				}
 			});
-			Pending.Add(MoveTemp(NormalFuture));
+			Pending.Add(MoveTemp(NormalTask));
 		}
 
 
 		if (TangentOverlay != nullptr)
 		{
-			auto TangentFuture = Async(EAsyncExecution::ThreadPool, [&]()
+			auto TangentTask = UE::Tasks::Launch(UE_SOURCE_LOCATION, [&]()
 			{
 				FNormalWelder TangentWelder;
 				TangentWelder.NormalOverlay = TangentOverlay;
@@ -715,12 +718,12 @@ void FMeshDescriptionToDynamicMesh::Convert(const FMeshDescription* MeshIn, FDyn
 					TangentOverlay->SetTriangle(TriangleID, TriVector);
 				}
 			});
-			Pending.Add(MoveTemp(TangentFuture));
+			Pending.Add(MoveTemp(TangentTask));
 		}
 
 		if (BiTangentOverlay != nullptr)
 		{
-			auto BiTangentFuture = Async(EAsyncExecution::ThreadPool, [&]()
+			auto BiTangentTask = UE::Tasks::Launch(UE_SOURCE_LOCATION, [&]()
 			{
 				FNormalWelder BiTangentWelder;
 				BiTangentWelder.NormalOverlay = BiTangentOverlay;
@@ -745,12 +748,12 @@ void FMeshDescriptionToDynamicMesh::Convert(const FMeshDescription* MeshIn, FDyn
 					BiTangentOverlay->SetTriangle(TriangleID, TriVector);
 				}
 			});
-			Pending.Add(MoveTemp(BiTangentFuture));
+			Pending.Add(MoveTemp(BiTangentTask));
 		}
 
 		if (ColorOverlay != nullptr)
 		{
-			auto ColorFuture = Async(EAsyncExecution::ThreadPool, [&]()
+			auto ColorTask = UE::Tasks::Launch(UE_SOURCE_LOCATION, [&]()
 			{
 				const FVector4f DefaultColor4 = InstanceColors.GetDefaultValue();
 
@@ -774,12 +777,12 @@ void FMeshDescriptionToDynamicMesh::Convert(const FMeshDescription* MeshIn, FDyn
 					ColorOverlay->SetTriangle(TriangleID, TriVector);
 				}
 			});
-			Pending.Add(MoveTemp(ColorFuture));
+			Pending.Add(MoveTemp(ColorTask));
 		}
 
 		if (MaterialIDAttrib != nullptr)
 		{
-			auto MaterialFuture = Async(EAsyncExecution::ThreadPool, [&]()
+			auto MaterialTask = UE::Tasks::Launch(UE_SOURCE_LOCATION, [&]()
 			{
 				for (int32 TriangleID : MeshOut.TriangleIndicesItr())
 				{
@@ -787,12 +790,12 @@ void FMeshDescriptionToDynamicMesh::Convert(const FMeshDescription* MeshIn, FDyn
 					MaterialIDAttrib->SetValue(TriangleID, &TriData.PolygonGroupID);
 				}
 			});
-			Pending.Add(MoveTemp(MaterialFuture));
+			Pending.Add(MoveTemp(MaterialTask));
 		}
 
 		for (const FSkinWeightsAttribCopyInfo& SkinWeightAttribInfo: SkinWeightAttribs)
 		{
-			auto SkinWeightsFuture = Async(EAsyncExecution::ThreadPool, [&, SkinWeightAttribInfo]() -> void
+			auto SkinWeightsTask = UE::Tasks::Launch(UE_SOURCE_LOCATION, [&, SkinWeightAttribInfo]() -> void
             {
                 for (int32 VertexID: MeshOut.VertexIndicesItr())
                 {
@@ -800,13 +803,13 @@ void FMeshDescriptionToDynamicMesh::Convert(const FMeshDescription* MeshIn, FDyn
                 }
 				
             });
-			Pending.Add(MoveTemp(SkinWeightsFuture));
+			Pending.Add(MoveTemp(SkinWeightsTask));
 		}
 		
 		// initialize polygroup layers
 		for (int32 GroupLayerIdx = 0; GroupLayerIdx < PolygroupAttribs.Num(); GroupLayerIdx++)
 		{
-			auto PolygroupFuture = Async(EAsyncExecution::ThreadPool, [&,Index=GroupLayerIdx]()
+			auto PolygroupTask = UE::Tasks::Launch(UE_SOURCE_LOCATION, [&,Index=GroupLayerIdx]()
 			{
 				TTriangleAttributesConstRef<int32> InputGroupSet = PolygroupAttribs[Index];
 				FDynamicMeshPolygroupAttribute* OutputGroupSet = MeshOut.Attributes()->GetPolygroupLayer(Index);
@@ -821,13 +824,13 @@ void FMeshDescriptionToDynamicMesh::Convert(const FMeshDescription* MeshIn, FDyn
 					}
 				}
 			});
-			Pending.Add(MoveTemp(PolygroupFuture));
+			Pending.Add(MoveTemp(PolygroupTask));
 		}
 
 		// initialize weight layers
 		for (int32 WeightLayerIdx = 0; WeightLayerIdx < WeightAttribs.Num(); WeightLayerIdx++)
 		{
-			auto WeightFuture = Async(EAsyncExecution::ThreadPool, [this, &WeightAttribs, &WeightAttribNames, &MeshOut, WeightLayerIdx]()
+			auto WeightTask = UE::Tasks::Launch(UE_SOURCE_LOCATION, [this, &WeightAttribs, &WeightAttribNames, &MeshOut, WeightLayerIdx]()
 			{
 				TVertexAttributesConstRef<float> InputWeights = WeightAttribs[WeightLayerIdx];
 				FDynamicMeshWeightAttribute* OutputWeights = MeshOut.Attributes()->GetWeightLayer(WeightLayerIdx);
@@ -842,14 +845,11 @@ void FMeshDescriptionToDynamicMesh::Convert(const FMeshDescription* MeshIn, FDyn
 					}
 				}
 			});
-			Pending.Add(MoveTemp(WeightFuture));
+			Pending.Add(MoveTemp(WeightTask));
 		}
 
 		// wait for all work to be done
-		for (TFuture<void>& Future : Pending)
-		{
-			Future.Wait();
-		}
+		UE::Tasks::Wait(Pending);
 
 		if (!bFoundNonDefaultVertexInstanceColor)
 		{
