Cloth Editor:
- implement a more robust ClothAsset to DynamicMesh converter using TToDynamicMesh
- save DynamicMesh weight maps to cloth asset using cloth adapter API
- remove old DynamicMesh/SkeletalMesh translation code

diff --git a/Engine/Source/Runtime/ClothingSystemRuntimeInterface/Public/ClothingAssetBase.h b/Engine/Source/Runtime/ClothingSystemRuntimeInterface/Public/ClothingAssetBase.h
--- a/Engine/Source/Runtime/ClothingSystemRuntimeInterface/Public/ClothingAssetBase.h
+++ b/Engine/Source/Runtime/ClothingSystemRuntimeInterface/Public/ClothingAssetBase.h
@@ -126,9 +126,6 @@ protected:
 	/** The asset factory should have access, as it will assign the asset guid when building assets */
 	friend class UClothingAssetFactory;
 
-	/* TODO: Remove this temporary access (JIRA UE-159698) */
-	friend class UChaosClothAssetEditorMode;	
-
 	/** Guid to identify this asset. Will be embedded into chunks that are created using this asset */
 	UPROPERTY()
 	FGuid AssetGuid;
diff --git a/Engine/Source/Runtime/MeshConversion/Public/ToDynamicMesh.h b/Engine/Source/Runtime/MeshConversion/Public/ToDynamicMesh.h
--- a/Engine/Source/Runtime/MeshConversion/Public/ToDynamicMesh.h
+++ b/Engine/Source/Runtime/MeshConversion/Public/ToDynamicMesh.h
@@ -362,6 +362,7 @@ protected:
 
 		// -- Create Overlays
 		const int32 NumUVLayers = MeshIn.NumUVLayers();
+		const int32 NumWeightMaps = MeshIn.NumWeightMapLayers();
 
 		MeshOut.EnableAttributes(); // by default 1-UV layer and 1-normal layer
 
@@ -384,7 +385,7 @@ protected:
 		}
 
 		// create UV overlays 
-		MeshOut.Attributes()->SetNumUVLayers(NumUVLayers);
+		MeshOut.Attributes()->SetNumUVLayers(FMath::Max(1, NumUVLayers));
 		// reserve space in any new UV layers.
 		for (int32 i = 1; i < NumUVLayers; ++i)
 		{
@@ -392,9 +393,12 @@ protected:
 		}
 
 		// create color overlay
-		MeshOut.Attributes()->EnablePrimaryColors();
-		ColorOverlay = MeshOut.Attributes()->PrimaryColors();
-		ColorOverlay->InitializeTriangles(MeshOut.MaxTriangleID());
+		if (MeshIn.HasColors())
+		{
+			MeshOut.Attributes()->EnablePrimaryColors();
+			ColorOverlay = MeshOut.Attributes()->PrimaryColors();
+			ColorOverlay->InitializeTriangles(MeshOut.MaxTriangleID());
+		}
 
 		// always enable Material ID if there are any attributes
 		MeshOut.Attributes()->EnableMaterialID();
@@ -504,6 +508,28 @@ protected:
 			Pending.Add(MoveTemp(MaterialFuture));
 		}
 
+		//populate WeightMap attribute
+		MeshOut.Attributes()->SetNumWeightLayers(NumWeightMaps);
+		for (int WeightMapIndex = 0; WeightMapIndex < NumWeightMaps; WeightMapIndex++)
+		{
+			auto WeightMapFuture = Async(EAsyncExecution::ThreadPool, [this, &MeshIn, &MeshOut, WeightMapIndex]()
+			{
+				FDynamicMeshWeightAttribute* WeightMapAttrib = MeshOut.Attributes()->GetWeightLayer(WeightMapIndex);
+
+				for (int32 VertexID : MeshOut.VertexIndicesItr())
+				{
+					SrcVertIDType SrcVertID = MyBase::ToSrcVertIDMap[VertexID];
+					float Weight = MeshIn.GetVertexWeight(WeightMapIndex, SrcVertID);
+					WeightMapAttrib->SetValue(VertexID, &Weight);
+				}
+
+				WeightMapAttrib->SetName( MeshIn.GetWeightMapName(WeightMapIndex) );
+			});
+			Pending.Add(MoveTemp(WeightMapFuture));
+		}
+
+		// TODO: PolyGroup layers
+
 		// wait for all work to be done
 		for (TFuture<void>& Future : Pending)
 		{
diff --git a/Engine/Source/Runtime/MeshConversionEngineTypes/Private/SkeletalMeshLODModelToDynamicMesh.cpp b/Engine/Source/Runtime/MeshConversionEngineTypes/Private/SkeletalMeshLODModelToDynamicMesh.cpp
--- a/Engine/Source/Runtime/MeshConversionEngineTypes/Private/SkeletalMeshLODModelToDynamicMesh.cpp
+++ b/Engine/Source/Runtime/MeshConversionEngineTypes/Private/SkeletalMeshLODModelToDynamicMesh.cpp
@@ -135,6 +135,23 @@ public:
 		return Mesh->NumTexCoords;
 	}
 
+	int32 NumWeightMapLayers() const
+	{
+		return 0;
+	}
+
+	float GetVertexWeight(int32, int32) const
+	{
+		ensure(false);
+		return 0.0f;
+	}
+
+	FName GetWeightMapName(int32 LayerIndex) const
+	{
+		return FName();
+	}
+
+
 	// --"Vertex Buffer" info
 	const TArray<VertIDType>& GetVertIDs() const
 	{
