ToDynamicMesh - update comments to reflect the required API wrt weightmap.

diff --git a/Engine/Source/Runtime/MeshConversion/Public/ToDynamicMesh.h b/Engine/Source/Runtime/MeshConversion/Public/ToDynamicMesh.h
--- a/Engine/Source/Runtime/MeshConversion/Public/ToDynamicMesh.h
+++ b/Engine/Source/Runtime/MeshConversion/Public/ToDynamicMesh.h
@@ -303,6 +303,11 @@ public:
 *		const TArray<ColorIDType>& GetColorIDs() const;
 *		FVector4f GetColor(ColorIDType ID) const;
 *		bool GetColorTri(const TriIDType& TID, ColorIDType& NID0, ColorIDType& NID1, ColorIDType& NID2) const;
+* 
+*       // weight maps information
+*	    int32 NumWeightMapLayers() const;
+*       float GetVertexWeight(int32 WeightMapIndex, int32 SrcVertID) const;
+*	    FName GetWeightMapName(int32 WeightMapIndex) const;
 *	};
 *
 */
