Handle bone data when converting to/from mesh description or to/from dynamic mesh.

diff --git a/Engine/Source/Runtime/Engine/Private/Rendering/SkeletalMeshLODImporterData.cpp b/Engine/Source/Runtime/Engine/Private/Rendering/SkeletalMeshLODImporterData.cpp
--- a/Engine/Source/Runtime/Engine/Private/Rendering/SkeletalMeshLODImporterData.cpp
+++ b/Engine/Source/Runtime/Engine/Private/Rendering/SkeletalMeshLODImporterData.cpp
@@ -1738,6 +1738,10 @@ bool FSkeletalMeshImportData::GetMeshDescription(FMeshDescription& OutMeshDescri
 
 	TPolygonGroupAttributesRef<FName> PolygonGroupMaterialSlotNames = MeshAttributes.GetPolygonGroupMaterialSlotNames();
 
+	FSkeletalMeshAttributes::FBoneNameAttributesRef BoneNames = MeshAttributes.GetBoneNames();
+	FSkeletalMeshAttributes::FBoneParentIndexAttributesRef BoneParentIndices = MeshAttributes.GetBoneParentIndices();
+	FSkeletalMeshAttributes::FBonePoseAttributesRef BonePoses = MeshAttributes.GetBonePoses();
+
 	VertexInstanceUVs.SetNumChannels(NumTexCoords);
 	
 	// Avoid repeated allocations and reserve the target buffers right off the bat.
@@ -1745,6 +1749,7 @@ bool FSkeletalMeshImportData::GetMeshDescription(FMeshDescription& OutMeshDescri
 	OutMeshDescription.ReserveNewTriangles(Faces.Num());
 	OutMeshDescription.ReserveNewVertexInstances(Wedges.Num());
 	OutMeshDescription.ReserveNewVertices(Points.Num());
+	MeshAttributes.ReserveNewBones(RefBonesBinary.Num());
 
 	// Copy the vertex positions first and maintain a map so that we can go from the import data's raw vertex index
 	// to the mesh description's VertexID.
@@ -1789,6 +1794,18 @@ bool FSkeletalMeshImportData::GetMeshDescription(FMeshDescription& OutMeshDescri
 		VertexSkinWeights.Set(VertexIDMap[VertexIndex], BoneWeights);		
 	}
 
+	// Set Bone Attributes
+	for (int Idx = 0; Idx < RefBonesBinary.Num(); ++Idx)
+	{
+		const SkeletalMeshImportData::FBone& Bone = RefBonesBinary[Idx];
+		
+		const FBoneID BoneID = MeshAttributes.CreateBone();
+		
+		BoneNames.Set(BoneID, FName(Bone.Name));
+		BoneParentIndices.Set(BoneID, Bone.ParentIndex);
+		BonePoses.Set(BoneID, FTransform(Bone.BonePos.Transform));
+	}
+
 	// Partition the faces by material index. Each material index corresponds to a polygon group.
 	// Here, it's worth doing via indirection, due to the size of each FTriangle object.
 	TArray<int32> FaceIndices;
diff --git a/Engine/Source/Runtime/MeshConversion/Private/DynamicMeshToMeshDescription.cpp b/Engine/Source/Runtime/MeshConversion/Private/DynamicMeshToMeshDescription.cpp
--- a/Engine/Source/Runtime/MeshConversion/Private/DynamicMeshToMeshDescription.cpp
+++ b/Engine/Source/Runtime/MeshConversion/Private/DynamicMeshToMeshDescription.cpp
@@ -5,6 +5,7 @@
 #include "DynamicMesh/DynamicMeshAttributeSet.h"
 #include "DynamicMesh/DynamicMeshOverlay.h"
 #include "DynamicMesh/DynamicVertexSkinWeightsAttribute.h"
+#include "DynamicMesh/DynamicBoneAttribute.h"
 #include "DynamicMesh/DynamicVertexAttribute.h"
 #include "MeshDescriptionBuilder.h"
 #include "DynamicMesh/MeshTangents.h"
@@ -678,6 +679,12 @@ void FDynamicMeshToMeshDescription::Convert_NoSharedInstances(const FDynamicMesh
 
 	const int32 NumUVLayers = bHasAttributes ? MeshIn->Attributes()->NumUVLayers() : 0;
 	
+	const bool bHasBones = bHasAttributes && MeshIn->Attributes()->HasBones();
+	const FDynamicMeshBoneNameAttribute* BoneNames = bHasBones ? MeshIn->Attributes()->GetBoneNames() : nullptr;
+	const FDynamicMeshBoneParentIndexAttribute* BoneParentIndices = bHasBones ? MeshIn->Attributes()->GetBoneParentIndices() : nullptr;
+	const FDynamicMeshBonePoseAttribute* BonePoses = bHasBones ? MeshIn->Attributes()->GetBonePoses() : nullptr;
+	const FDynamicMeshBoneColorAttribute* BoneColors = bHasBones ? MeshIn->Attributes()->GetBoneColors() : nullptr;
+
 	// cache the UV layers
 	TArray<const FDynamicMeshUVOverlay*> UVLayers;
 	for (int32 k = 0; k < NumUVLayers; ++k)
@@ -698,11 +705,36 @@ void FDynamicMeshToMeshDescription::Convert_NoSharedInstances(const FDynamicMesh
 		bCopyGroupToPolyGroup = true;
 	}
 
+	// We register skeletal attributes if either bone names or skinning infomation is available
+	if (bHasBones || !MeshIn->Attributes()->GetSkinWeightsAttributes().IsEmpty())
+	{
+		FSkeletalMeshAttributes MeshOutAttributes(MeshOut);
+		MeshOutAttributes.Register();
+		MeshOutAttributes.RegisterColorAttribute();
+	}
+
+	if (bHasBones)
+	{
+		FSkeletalMeshAttributes MeshOutAttributes(MeshOut);
+		checkSlow(MeshOutAttributes.HasBones()); // FSkeletalMeshAttributes::Register() must have been called
+		
+		const int32 NumBones = MeshIn->Attributes()->GetNumBones();
+		MeshOutAttributes.ReserveNewBones(NumBones);
+
+		for (int Idx = 0; Idx < NumBones; ++Idx)
+		{
+			const FBoneID BoneID = MeshOutAttributes.CreateBone();
+			MeshOutAttributes.GetBoneNames().Set(BoneID, BoneNames->GetValue(Idx));
+			MeshOutAttributes.GetBoneParentIndices().Set(BoneID, BoneParentIndices->GetValue(Idx));
+			MeshOutAttributes.GetBonePoses().Set(BoneID, BonePoses->GetValue(Idx));
+			MeshOutAttributes.GetBoneColors().Set(BoneID, BoneColors->GetValue(Idx));
+		}
+	}
+
 	TMap<FName, FSkinWeightsVertexAttributesRef> VertexBoneWeightsMap; 
 	if (!MeshIn->Attributes()->GetSkinWeightsAttributes().IsEmpty())
 	{
 		FSkeletalMeshAttributes MeshOutAttributes(MeshOut);
-		MeshOutAttributes.Register();
 		for (const TTuple<FName, TUniquePtr<FDynamicMeshVertexSkinWeightsAttribute>>& AttributeInfo: MeshIn->Attributes()->GetSkinWeightsAttributes())
 		{
 			FName ProfileName = AttributeInfo.Key;
diff --git a/Engine/Source/Runtime/MeshConversion/Private/MeshDescriptionToDynamicMesh.cpp b/Engine/Source/Runtime/MeshConversion/Private/MeshDescriptionToDynamicMesh.cpp
--- a/Engine/Source/Runtime/MeshConversion/Private/MeshDescriptionToDynamicMesh.cpp
+++ b/Engine/Source/Runtime/MeshConversion/Private/MeshDescriptionToDynamicMesh.cpp
@@ -6,6 +6,7 @@
 #include "DynamicMesh/DynamicMeshOverlay.h"
 #include "DynamicMesh/DynamicVertexAttribute.h"
 #include "DynamicMesh/DynamicVertexSkinWeightsAttribute.h"
+#include "DynamicMesh/DynamicBoneAttribute.h"
 #include "DynamicMesh/MeshTangents.h"
 #include "DynamicMesh/NonManifoldMappingSupport.h"
 #include "MeshDescriptionBuilder.h"
@@ -399,6 +400,10 @@ void FMeshDescriptionToDynamicMesh::Convert(const FMeshDescription* MeshIn, FDyn
 	FDynamicMeshNormalOverlay* TangentOverlay = nullptr;
 	FDynamicMeshNormalOverlay* BiTangentOverlay = nullptr;
 	FDynamicMeshMaterialAttribute* MaterialIDAttrib = nullptr;
+	FDynamicMeshBoneNameAttribute* BoneNameAttrib = nullptr;
+	FDynamicMeshBoneParentIndexAttribute* BoneParentIndexAttrib = nullptr;
+	FDynamicMeshBonePoseAttribute* BonePoseAttrib = nullptr;
+	FDynamicMeshBoneColorAttribute* BoneColorAttrib = nullptr;
 
 	TArray<FSkinWeightsAttribCopyInfo> SkinWeightAttribs;
 	if (!bDisableAttributes)
@@ -460,6 +465,30 @@ void FMeshDescriptionToDynamicMesh::Convert(const FMeshDescription* MeshIn, FDyn
 				MeshOut.Attributes()->AttachSkinWeightsAttribute(ProfileName, VertexSkinWeightsAttrib);
 			}
 		});
+
+		if (Attributes.HasBones())
+		{	
+			MeshOut.Attributes()->EnableBones(Attributes.GetNumBones());
+			if (Attributes.HasBoneNameAttribute())
+			{
+				BoneNameAttrib = MeshOut.Attributes()->GetBoneNames();
+			}
+			
+			if (Attributes.HasBoneParentIndexAttribute())
+			{
+				BoneParentIndexAttrib = MeshOut.Attributes()->GetBoneParentIndices();
+			}
+
+			if (Attributes.HasBonePoseAttribute())
+			{
+				BonePoseAttrib = MeshOut.Attributes()->GetBonePoses();
+			}
+
+			if (Attributes.HasBoneColorAttribute())
+			{
+				BoneColorAttrib = MeshOut.Attributes()->GetBoneColors();
+			}
+		}
 	}
 
 
@@ -809,7 +838,36 @@ void FMeshDescriptionToDynamicMesh::Convert(const FMeshDescription* MeshIn, FDyn
             });
 			Pending.Add(MoveTemp(SkinWeightsTask));
 		}
-		
+
+		if (BoneNameAttrib != nullptr)
+		{	
+			auto BonesTask = UE::Tasks::Launch(UE_SOURCE_LOCATION, [&]() -> void
+			{
+				const int NumBones = Attributes.GetNumBones();
+				for (int32 BoneID = 0; BoneID < NumBones; ++BoneID)
+				{
+					BoneNameAttrib->SetValue(BoneID, Attributes.GetBoneNames().Get(BoneID));
+
+					if (BoneParentIndexAttrib)
+					{	
+						BoneParentIndexAttrib->SetValue(BoneID, Attributes.GetBoneParentIndices().Get(BoneID));
+					}
+
+					if (BonePoseAttrib)
+					{
+						BonePoseAttrib->SetValue(BoneID, Attributes.GetBonePoses().Get(BoneID));
+					}
+
+					if (BoneColorAttrib)
+					{
+						BoneColorAttrib->SetValue(BoneID, Attributes.GetBoneColors().Get(BoneID));
+					}
+				}
+				
+			});
+			Pending.Add(MoveTemp(BonesTask));
+		}
+
 		// initialize polygroup layers
 		for (int32 GroupLayerIdx = 0; GroupLayerIdx < PolygroupAttribs.Num(); GroupLayerIdx++)
 		{
