Geometry: Fixed issue where a deference could be referring to a null value when determining number of UV channels during mesh conversion.

diff --git a/Engine/Source/Runtime/MeshConversion/Private/MeshDescriptionToDynamicMesh.cpp b/Engine/Source/Runtime/MeshConversion/Private/MeshDescriptionToDynamicMesh.cpp
--- a/Engine/Source/Runtime/MeshConversion/Private/MeshDescriptionToDynamicMesh.cpp
+++ b/Engine/Source/Runtime/MeshConversion/Private/MeshDescriptionToDynamicMesh.cpp
@@ -378,7 +378,7 @@ void FMeshDescriptionToDynamicMesh::Convert(const FMeshDescription* MeshIn, FDyn
 	const int32 NumUVElementChannels = MeshIn->GetNumUVElementChannels();
 
 	// the instanced UV representation.
-	const int NumUVLayers = InstanceUVs.GetNumChannels();
+	const int NumUVLayers = InstanceUVs.IsValid() ? InstanceUVs.GetNumChannels() : 0;
 
 	// determine if we really have shared UVS. Legacy geo might not have them 
 	// - at the time of this writing MeshDescription has not been updated
