ToDynamicMesh: Use SrcAttr indices when querying SharedAttrValues in PopulateOverlayFromSharedAttr

diff --git a/Engine/Source/Runtime/MeshConversion/Public/ToDynamicMesh.h b/Engine/Source/Runtime/MeshConversion/Public/ToDynamicMesh.h
--- a/Engine/Source/Runtime/MeshConversion/Public/ToDynamicMesh.h
+++ b/Engine/Source/Runtime/MeshConversion/Public/ToDynamicMesh.h
@@ -719,7 +719,7 @@ protected:
 					int32 ParentVID = Overlay->GetParentVertex(AttrTri[i]);
 					if (ParentVID != FDynamicMesh3::InvalidID && ParentVID != ParentTriangle[i])
 					{
-						const AttrType AttrValue = SharedAttrValues(AttrTri[i]);
+						const AttrType AttrValue = SharedAttrValues(SrcAttrTri[i]);
 						AttrTri[i] = Overlay->AppendElement(AttrValue);
 					}
 				}
@@ -730,7 +730,7 @@ protected:
 
 				if (AttrTri[0] == AttrTri[1] && AttrTri[0] == AttrTri[2])
 				{
-					const AttrType AttrValue = SharedAttrValues(AttrTri[0]);
+					const AttrType AttrValue = SharedAttrValues(SrcAttrTri[0]);
 					AttrTri[0] = Overlay->AppendElement(AttrValue);
 					AttrTri[1] = Overlay->AppendElement(AttrValue);
 					AttrTri[2] = Overlay->AppendElement(AttrValue);
@@ -740,19 +740,19 @@ protected:
 				{
 					if (AttrTri[0] == AttrTri[1])
 					{
-						const AttrType AttrValue = SharedAttrValues(AttrTri[0]);
+						const AttrType AttrValue = SharedAttrValues(SrcAttrTri[0]);
 						AttrTri[0] = Overlay->AppendElement(AttrValue);
 						AttrTri[1] = Overlay->AppendElement(AttrValue);
 					}
 					if (AttrTri[0] == AttrTri[2])
 					{
-						const AttrType AttrValue = SharedAttrValues(AttrTri[0]);
+						const AttrType AttrValue = SharedAttrValues(SrcAttrTri[0]);
 						AttrTri[0] = Overlay->AppendElement(AttrValue);
 						AttrTri[2] = Overlay->AppendElement(AttrValue);
 					}
 					if (AttrTri[1] == AttrTri[2])
 					{
-						const AttrType AttrValue = SharedAttrValues(AttrTri[1]);
+						const AttrType AttrValue = SharedAttrValues(SrcAttrTri[1]);
 						AttrTri[1] = Overlay->AppendElement(AttrValue);
 						AttrTri[2] = Overlay->AppendElement(AttrValue);
 					}
