Fix ToDynamicMesh's ToSrcVertIDMap when splitting vertices for non-manifold meshes.

diff --git a/Engine/Source/Runtime/MeshConversion/Public/ToDynamicMesh.h b/Engine/Source/Runtime/MeshConversion/Public/ToDynamicMesh.h
--- a/Engine/Source/Runtime/MeshConversion/Public/ToDynamicMesh.h
+++ b/Engine/Source/Runtime/MeshConversion/Public/ToDynamicMesh.h
@@ -156,7 +156,7 @@ public:
 						// grow the map to make room
 						ToSrcVertIDMap.SetNumUninitialized(NewDstVertID + 1);
 
-						ToSrcVertIDMap[NewDstVertID] = DstTriVerts[i]; 
+						ToSrcVertIDMap[NewDstVertID] = SrcTriVerts[i];
 					}
 				}
 
