Disable false positives for
V557: Array overrun is possible. The 'GetOnPhaseStart' function processes value '[0..6]'. Inspect the first argument.
V621 Consider inspecting the 'for' operator. It's possible that the loop will be executed incorrectly or won't be executed at all.
V654 The condition loop is always false

diff --git a/Engine/Plugins/Runtime/MassEntity/Source/MassEntity/Public/MassProcessingPhaseManager.h b/Engine/Plugins/Runtime/MassEntity/Source/MassEntity/Public/MassProcessingPhaseManager.h
--- a/Engine/Plugins/Runtime/MassEntity/Source/MassEntity/Public/MassProcessingPhaseManager.h
+++ b/Engine/Plugins/Runtime/MassEntity/Source/MassEntity/Public/MassProcessingPhaseManager.h
@@ -136,7 +136,7 @@ public:
 	FMassEntityManager& GetEntityManagerRef() { check(EntityManager); return *EntityManager.Get(); }
 
 	/** Retrieves OnPhaseStart multicast delegate's reference for a given Phase */
-	FMassProcessingPhase::FOnPhaseEvent& GetOnPhaseStart(const EMassProcessingPhase Phase) { return ProcessingPhases[uint8(Phase)].OnPhaseStart; }
+	FMassProcessingPhase::FOnPhaseEvent& GetOnPhaseStart(const EMassProcessingPhase Phase) { return ProcessingPhases[uint8(Phase)].OnPhaseStart; } //-V557
 	/** Retrieves OnPhaseEnd multicast delegate's reference for a given Phase */
 	FMassProcessingPhase::FOnPhaseEvent& GetOnPhaseEnd(const EMassProcessingPhase Phase) { return ProcessingPhases[uint8(Phase)].OnPhaseEnd; }
 
diff --git a/Engine/Source/Runtime/Experimental/IoStoreOnDemand/Private/OnDemandIoDispatcherBackend.cpp b/Engine/Source/Runtime/Experimental/IoStoreOnDemand/Private/OnDemandIoDispatcherBackend.cpp
--- a/Engine/Source/Runtime/Experimental/IoStoreOnDemand/Private/OnDemandIoDispatcherBackend.cpp
+++ b/Engine/Source/Runtime/Experimental/IoStoreOnDemand/Private/OnDemandIoDispatcherBackend.cpp
@@ -2417,7 +2417,7 @@ void FOnDemandIoBackend::ProcessHttpRequests(FHttpClient& HttpClient, FBitWindow
 			if (NumConcurrentRequests >= MaxConcurrentRequests)
 			{
 				TRACE_CPUPROFILER_EVENT_SCOPE(IasBackend::TickHttpSaturated);
-				while (NumConcurrentRequests >= MaxConcurrentRequests)
+				while (NumConcurrentRequests >= MaxConcurrentRequests) //-V654
 				{
 					HttpClient.Tick(MAX_uint32, GIasHttpRateLimitKiBPerSecond);
 				}
diff --git a/Engine/Source/Runtime/MeshConversion/Public/ToDynamicMesh.h b/Engine/Source/Runtime/MeshConversion/Public/ToDynamicMesh.h
--- a/Engine/Source/Runtime/MeshConversion/Public/ToDynamicMesh.h
+++ b/Engine/Source/Runtime/MeshConversion/Public/ToDynamicMesh.h
@@ -552,7 +552,7 @@ protected:
 		const int32 NumSkinWeightAttributes = MeshIn.NumSkinWeightAttributes();
 		
 		// first, attach all the skinning weight attributes (this also allocates memory to store skinning weights for each attribute)
-		for (int AttributeIndex = 0; AttributeIndex < NumSkinWeightAttributes; ++AttributeIndex) 
+		for (int AttributeIndex = 0; AttributeIndex < NumSkinWeightAttributes; ++AttributeIndex) //-V621 //-V654
 		{
 			UE::Geometry::FDynamicMeshVertexSkinWeightsAttribute* Attribute = new UE::Geometry::FDynamicMeshVertexSkinWeightsAttribute(&MeshOut);
 			const FName AttribName = MeshIn.GetSkinWeightAttributeName(AttributeIndex);
@@ -561,7 +561,7 @@ protected:
 		}
 
 		// now set all of the skin weight data
-		for (int AttributeIndex = 0; AttributeIndex < NumSkinWeightAttributes; ++AttributeIndex)
+		for (int AttributeIndex = 0; AttributeIndex < NumSkinWeightAttributes; ++AttributeIndex) //-V621
 		{
 			auto SkinAttributeFeature = Async(EAsyncExecution::ThreadPool, [this, &MeshIn, &MeshOut, AttributeIndex]()
 			{
