MeshConversion - DynamicMeshToMeshDescription conversion: Fix crash occuring when input mesh has no attributes

diff --git a/Engine/Source/Runtime/MeshConversion/Private/DynamicMeshToMeshDescription.cpp b/Engine/Source/Runtime/MeshConversion/Private/DynamicMeshToMeshDescription.cpp
--- a/Engine/Source/Runtime/MeshConversion/Private/DynamicMeshToMeshDescription.cpp
+++ b/Engine/Source/Runtime/MeshConversion/Private/DynamicMeshToMeshDescription.cpp
@@ -716,6 +716,8 @@ void FDynamicMeshToMeshDescription::Convert_NoSharedInstances(const FDynamicMesh
 	const FDynamicMeshBonePoseAttribute* BonePoses = bHasBones ? MeshIn->Attributes()->GetBonePoses() : nullptr;
 	const FDynamicMeshBoneColorAttribute* BoneColors = bHasBones ? MeshIn->Attributes()->GetBoneColors() : nullptr;
 
+	const bool bHasSkinWeightsAttributes = bHasAttributes && !MeshIn->Attributes()->GetSkinWeightsAttributes().IsEmpty();
+
 	// cache the UV layers
 	TArray<const FDynamicMeshUVOverlay*> UVLayers;
 	for (int32 k = 0; k < NumUVLayers; ++k)
@@ -737,7 +739,7 @@ void FDynamicMeshToMeshDescription::Convert_NoSharedInstances(const FDynamicMesh
 	}
 
 	// We register skeletal attributes if either bone names or skinning infomation is available
-	if (bHasBones || !MeshIn->Attributes()->GetSkinWeightsAttributes().IsEmpty())
+	if (bHasBones || bHasSkinWeightsAttributes)
 	{
 		FSkeletalMeshAttributes MeshOutAttributes(MeshOut);
 		MeshOutAttributes.Register();
@@ -771,7 +773,7 @@ void FDynamicMeshToMeshDescription::Convert_NoSharedInstances(const FDynamicMesh
 	}
 
 	TMap<FName, FSkinWeightsVertexAttributesRef> VertexBoneWeightsMap; 
-	if (!MeshIn->Attributes()->GetSkinWeightsAttributes().IsEmpty())
+	if (bHasSkinWeightsAttributes)
 	{
 		FSkeletalMeshAttributes MeshOutAttributes(MeshOut);
 		for (const TTuple<FName, TUniquePtr<FDynamicMeshVertexSkinWeightsAttribute>>& AttributeInfo: MeshIn->Attributes()->GetSkinWeightsAttributes())
@@ -1023,20 +1025,23 @@ void FDynamicMeshToMeshDescription::Convert_NoSharedInstances(const FDynamicMesh
 
 	// Convert all attached skin weights, if we're converting a mesh description that originated from
 	// a USkeletalMesh.
-	for (const TTuple<FName, TUniquePtr<FDynamicMeshVertexSkinWeightsAttribute>>& AttributeInfo: MeshIn->Attributes()->GetSkinWeightsAttributes())
+	if (bHasSkinWeightsAttributes)
 	{
-		FName ProfileName = AttributeInfo.Key;
+		for (const TTuple<FName, TUniquePtr<FDynamicMeshVertexSkinWeightsAttribute>>& AttributeInfo: MeshIn->Attributes()->GetSkinWeightsAttributes())
+		{
+			FName ProfileName = AttributeInfo.Key;
 		
-		FSkinWeightsVertexAttributesRef& VertexBoneWeights = VertexBoneWeightsMap[ProfileName];
+			FSkinWeightsVertexAttributesRef& VertexBoneWeights = VertexBoneWeightsMap[ProfileName];
 
-		const FDynamicMeshVertexSkinWeightsAttribute *MeshSkinWeights = AttributeInfo.Value.Get();
-		for (int32 VertexIndex = 0; VertexIndex < MapV.Num(); VertexIndex++)
-		{
-			if (const FVertexID VertexID = MapV[VertexIndex]; VertexID != INDEX_NONE)
+			const FDynamicMeshVertexSkinWeightsAttribute *MeshSkinWeights = AttributeInfo.Value.Get();
+			for (int32 VertexIndex = 0; VertexIndex < MapV.Num(); VertexIndex++)
 			{
-				UE::AnimationCore::FBoneWeights BW;
-				MeshSkinWeights->GetValue(VertexIndex, BW);
-				VertexBoneWeights.Set(VertexID, BW);
+				if (const FVertexID VertexID = MapV[VertexIndex]; VertexID != INDEX_NONE)
+				{
+					UE::AnimationCore::FBoneWeights BW;
+					MeshSkinWeights->GetValue(VertexIndex, BW);
+					VertexBoneWeights.Set(VertexID, BW);
+				}
 			}
 		}
 	}
