Geometry Script:  Add ability to create a dynamic mesh from skeletal mesh render data ( used for auto-generated lods)

diff --git a/Engine/Plugins/Runtime/GeometryScripting/Source/GeometryScriptingCore/Private/MeshAssetFunctions.cpp b/Engine/Plugins/Runtime/GeometryScripting/Source/GeometryScriptingCore/Private/MeshAssetFunctions.cpp
--- a/Engine/Plugins/Runtime/GeometryScripting/Source/GeometryScriptingCore/Private/MeshAssetFunctions.cpp
+++ b/Engine/Plugins/Runtime/GeometryScripting/Source/GeometryScriptingCore/Private/MeshAssetFunctions.cpp
@@ -16,7 +16,9 @@
 #include "RenderingThread.h"
 #include "PhysicsEngine/BodySetup.h"
 #include "Rendering/SkeletalMeshModel.h"
+#include "Rendering/SkeletalMeshRenderData.h"
 #include "SkeletalMeshAttributes.h"
+#include "SkeletalMeshLODRenderDataToDynamicMesh.h"
 
 #include "MeshDescriptionToDynamicMesh.h"
 #include "DynamicMeshToMeshDescription.h"
@@ -536,6 +538,41 @@ void UGeometryScriptLibrary_StaticMeshFunctions::GetSectionMaterialListFromStati
 }
 
 
+namespace UELocal
+{
+	bool CopyMeshFromSkeletalMesh_RenderData(USkeletalMesh* FromSkeletalMeshAsset, FGeometryScriptCopyMeshFromAssetOptions AssetOptions, EGeometryScriptLODType LODType, int32 LODIndex, UDynamicMesh* ToDynamicMesh, UGeometryScriptDebug* Debug)
+	{
+		
+
+		if (FSkeletalMeshRenderData* RenderData = FromSkeletalMeshAsset->GetResourceForRendering())
+		{
+			int32 NumLODs = RenderData->LODRenderData.Num();
+			if (NumLODs -1 < LODIndex )
+			{
+				UE::Geometry::AppendError(Debug, EGeometryScriptErrorType::InvalidInputs, LOCTEXT("CopyMeshFromSkeletalMesh_RenderDataLDONotAvailable", "CopyMeshFromSkeletalMesh: Renderdata for specified LOD is not available"));
+				return false;
+			}
+
+			FSkeletalMeshLODRenderData* SkeletalMeshLODRenderData = &(RenderData->LODRenderData[LODIndex]);
+
+			UE::Geometry::FDynamicMesh3 NewMesh;
+			
+			UE::Geometry::FSkeletalMeshLODRenderDataToDynamicMesh::ConversionOptions  ConversionOptions;
+			ConversionOptions.bWantTangents = AssetOptions.bRequestTangents;
+
+			UE::Geometry::FSkeletalMeshLODRenderDataToDynamicMesh::Convert(SkeletalMeshLODRenderData, FromSkeletalMeshAsset->GetRefSkeleton(), ConversionOptions, NewMesh);
+			ToDynamicMesh->SetMesh(MoveTemp(NewMesh));
+
+			return true;
+		}
+	
+		UE::Geometry::AppendError(Debug, EGeometryScriptErrorType::InvalidInputs, LOCTEXT("CopyMeshFromSkeletalMesh_RenderDataNotAvailable", "CopyMeshFromSkeletalMesh: Renderdata is not available"));
+		return false;
+		
+	}
+	
+};
+
 UDynamicMesh* UGeometryScriptLibrary_StaticMeshFunctions::CopyMeshFromSkeletalMesh(
 		USkeletalMesh* FromSkeletalMeshAsset, 
 		UDynamicMesh* ToDynamicMesh,
@@ -556,37 +593,50 @@ UDynamicMesh* UGeometryScriptLibrary_StaticMeshFunctions::CopyMeshFromSkeletalMe
 		UE::Geometry::AppendError(Debug, EGeometryScriptErrorType::InvalidInputs, LOCTEXT("CopyMeshFromSkeletalMesh_InvalidInput2", "CopyMeshFromSkeletalMesh: ToDynamicMesh is Null"));
 		return ToDynamicMesh;
 	}
-	if (RequestedLOD.LODType != EGeometryScriptLODType::MaxAvailable && RequestedLOD.LODType != EGeometryScriptLODType::SourceModel)
-	{
-		UE::Geometry::AppendError(Debug, EGeometryScriptErrorType::InvalidInputs, LOCTEXT("CopyMeshFromSkeletalMesh_LODNotAvailable", "CopyMeshFromSkeletalMesh: Requested LOD is not available"));
-		return ToDynamicMesh;
-	}
 
 	// TODO: Consolidate this code with SkeletalMeshToolTarget::GetMeshDescription(..) 
-#if WITH_EDITOR
 	const int32 UseLODIndex = FMath::Clamp(RequestedLOD.LODIndex, 0, FromSkeletalMeshAsset->GetLODNum() - 1);;
 	
-	const FMeshDescription* SourceMesh = nullptr;
-
-	// Check first if we have bulk data available and non-empty.
-	if (FromSkeletalMeshAsset->HasMeshDescription(UseLODIndex))
+	if (RequestedLOD.LODType == EGeometryScriptLODType::MaxAvailable || RequestedLOD.LODType == EGeometryScriptLODType::SourceModel)
 	{
-		SourceMesh = FromSkeletalMeshAsset->GetMeshDescription(UseLODIndex); 
+#if WITH_EDITOR
+		const FMeshDescription* SourceMesh = nullptr;
+
+		// Check first if we have bulk data available and non-empty.
+		if (FromSkeletalMeshAsset->HasMeshDescription(UseLODIndex))
+		{
+			SourceMesh = FromSkeletalMeshAsset->GetMeshDescription(UseLODIndex); 
+		}
+		if (SourceMesh == nullptr)
+		{
+			UE::Geometry::AppendError(Debug, EGeometryScriptErrorType::InvalidInputs, LOCTEXT("CopyMeshFromSkeletalMesh_LODNotAvailable", "CopyMeshFromSkeletalMesh: Requested LOD source mesh is not available"));
+			return ToDynamicMesh;
+		}
+
+		FDynamicMesh3 NewMesh;
+		FMeshDescriptionToDynamicMesh Converter;
+		Converter.Convert(SourceMesh, NewMesh, AssetOptions.bRequestTangents);
+	
+		ToDynamicMesh->SetMesh(MoveTemp(NewMesh));
+	
+		Outcome = EGeometryScriptOutcomePins::Success;
+#else 
+		UE::Geometry::AppendError(Debug, EGeometryScriptErrorType::InvalidInputs, LOCTEXT("CopyMeshFromSkeletalMesh_SourceMesh_EditorOnly", "CopyMeshFromSkeletalMesh: Source Meshes are not available at Runtime"));
+	
+#endif	
 	}
-	if (SourceMesh == nullptr)
+	else if (RequestedLOD.LODType == EGeometryScriptLODType::RenderData)
 	{
-		UE::Geometry::AppendError(Debug, EGeometryScriptErrorType::InvalidInputs, LOCTEXT("CopyMeshFromSkeletalMesh_LODNotAvailable", "CopyMeshFromSkeletalMesh: Requested LOD is not available"));
-		return ToDynamicMesh;
+		if (UELocal::CopyMeshFromSkeletalMesh_RenderData(FromSkeletalMeshAsset, AssetOptions, RequestedLOD.LODType, RequestedLOD.LODIndex, ToDynamicMesh, Debug))
+		{
+			Outcome = EGeometryScriptOutcomePins::Success;
+		}
+		
 	}
 
-	FDynamicMesh3 NewMesh;
-	FMeshDescriptionToDynamicMesh Converter;
-	Converter.Convert(SourceMesh, NewMesh, AssetOptions.bRequestTangents);
-	
-	ToDynamicMesh->SetMesh(MoveTemp(NewMesh));
-	
-	Outcome = EGeometryScriptOutcomePins::Success;
-#else
+// todo remove this
+#if !WITH_EDITOR
+
 	UE::Geometry::AppendError(Debug, EGeometryScriptErrorType::InvalidInputs, LOCTEXT("CopyMeshFromSkeletalMesh_EditorOnly", "CopyMeshFromSkeletalMesh: Not currently supported at Runtime"));
 #endif
 	
diff --git a/Engine/Source/Runtime/MeshConversion/Public/SkeletalMeshLODRenderDataMeshAdapter.h b/Engine/Source/Runtime/MeshConversion/Public/SkeletalMeshLODRenderDataMeshAdapter.h
new file mode 100644
--- /dev/null
+++ b/Engine/Source/Runtime/MeshConversion/Public/SkeletalMeshLODRenderDataMeshAdapter.h
@@ -0,0 +1,328 @@
+// Copyright Epic Games, Inc.All Rights Reserved.
+
+#pragma once
+
+
+
+#include "CoreMinimal.h"
+#include "TransformTypes.h"
+#include "Rendering\SkeletalMeshLODRenderData.h"
+
+/**
+ * Basic struct to adapt a FSkeletalMeshLODRenderData for use by GeometryProcessing classes that template the mesh type and expect a standard set of basic accessors
+ * For example, this adapter will let you use a FSkeletalMeshLODRenderData with GeometryProcessing's TMeshAABBTree3
+ */
+struct FSkeletalMeshLODRenderDataMeshAdapter
+{
+	using FIndex3i = UE::Geometry::FIndex3i;
+protected:
+	const FSkeletalMeshLODRenderData* Mesh;
+
+	FVector3d BuildScale = FVector3d::One();
+	FVector3d InvBuildScale = FVector3d::One();
+	bool bScaleNormals = false;
+
+	TArray<const FSkelMeshRenderSection*> ValidSections;
+	TArray<int32> TriangleOffsetArray;
+	int32 NumTriangles;
+
+	TArray<uint32> SrcIndexBuffer;
+
+	TArray<FSkinWeightInfo> SkinWeights;
+
+protected:
+	FSkeletalMeshLODRenderDataMeshAdapter()
+		: Mesh(nullptr)
+		, NumTriangles(0)
+	{
+	}
+
+public:
+	FSkeletalMeshLODRenderDataMeshAdapter(const FSkeletalMeshLODRenderData* MeshIn)
+		: Mesh(MeshIn)
+		, NumTriangles(0)
+	{
+		if (MeshIn)
+		{
+			TriangleOffsetArray.Reserve(MeshIn->RenderSections.Num() + 1);
+			ValidSections.Reserve(MeshIn->RenderSections.Num());
+
+			for (const FSkelMeshRenderSection& Section : MeshIn->RenderSections)
+			{
+				TriangleOffsetArray.Add(NumTriangles);
+				NumTriangles += Section.NumTriangles;
+				ValidSections.Add(&Section);
+			}
+
+			TriangleOffsetArray.Add(NumTriangles);
+
+			MeshIn->MultiSizeIndexContainer.GetIndexBuffer(SrcIndexBuffer);
+
+			const FSkinWeightVertexBuffer* SkinWeightBuffer = &MeshIn->SkinWeightVertexBuffer;
+			const bool bSkinWeightsValid = SkinWeightBuffer->GetDataVertexBuffer() && SkinWeightBuffer->GetDataVertexBuffer()->IsInitialized() && SkinWeightBuffer->GetDataVertexBuffer()->GetNumVertices() > 0;
+			if (bSkinWeightsValid)
+			{
+				const int32 NumVerts = Mesh && Mesh->StaticVertexBuffers.PositionVertexBuffer.GetAllowCPUAccess() ? Mesh->GetNumVertices() : 0;
+
+				SkinWeights.Reserve(NumVerts);
+				for (const FSkelMeshRenderSection& Section : MeshIn->RenderSections)
+				{
+					const int32 BaseVertexIndex = static_cast<int32>(Section.BaseVertexIndex);
+					const int32 NumSectionVtx = static_cast<int32>(Section.NumVertices);
+					for (int32 VtxIndex = BaseVertexIndex; VtxIndex < NumSectionVtx + BaseVertexIndex; ++VtxIndex)
+					{
+						FSkinWeightInfo SrcWeights = SkinWeightBuffer->GetVertexSkinWeights(VtxIndex);
+
+						for (int32 Idx = 0; Idx < MAX_TOTAL_INFLUENCES; ++Idx)
+						{
+							const int32 SectionBoneIdx = static_cast<int32>(SrcWeights.InfluenceBones[Idx]);
+							if (ensure(SectionBoneIdx < Section.BoneMap.Num()))
+							{
+								SrcWeights.InfluenceBones[Idx] = Section.BoneMap[SectionBoneIdx];
+							}
+						}
+						SkinWeights.Add(SrcWeights);
+					}
+				}
+			}
+		}
+	}
+
+	void SetBuildScale(const FVector3d& BuildScaleIn, bool bScaleNormalsIn)
+	{
+		BuildScale = BuildScaleIn;
+		InvBuildScale = UE::Geometry::FTransformSRT3d::GetSafeScaleReciprocal(BuildScaleIn);
+		bScaleNormals = bScaleNormalsIn;
+	}
+
+	bool IsTriangle(int32 TID) const
+	{
+		return TID >= 0 && TID < MaxTriangleID();
+	}
+	bool IsVertex(int32 VID) const
+	{
+		return VID >= 0 && VID < MaxVertexID();
+	}
+	// ID and Count are the same for StaticMeshLODResources because it's compact
+	int32 MaxTriangleID() const
+	{
+		return TriangleCount();
+	}
+	int32 TriangleCount() const
+	{
+		return Mesh && Mesh->MultiSizeIndexContainer.IsIndexBufferValid() ? NumTriangles : 0;
+	}
+	int32 MaxVertexID() const
+	{
+		return VertexCount();
+	}
+	int32 VertexCount() const
+	{
+		return Mesh && Mesh->StaticVertexBuffers.PositionVertexBuffer.GetAllowCPUAccess() ? Mesh->GetNumVertices() : 0;
+	}
+	uint64 GetChangeStamp() const
+	{
+		// FSkeletalMeshLODRenderData doesn't provide any mechanism to know if it's been modified so just return 1
+		// and leave it to the caller to not build an aabb and then change the underlying mesh
+		return 1;
+	}
+
+	inline const FSkelMeshRenderSection& TriangleToSection(int32& InOutIDValue) const
+	{
+		int32 SectionIdx = Algo::UpperBound(TriangleOffsetArray, InOutIDValue) - 1;
+		InOutIDValue -= TriangleOffsetArray[SectionIdx];
+		return *ValidSections[SectionIdx];
+	}
+
+	inline FIndex3i GetTriangle(int32 IDValue) const
+	{
+		const FSkelMeshRenderSection& TriSection = TriangleToSection(IDValue);
+
+		return FIndex3i(SrcIndexBuffer[TriSection.BaseIndex + IDValue * 3 + 0],
+			SrcIndexBuffer[TriSection.BaseIndex + IDValue * 3 + 1],
+			SrcIndexBuffer[TriSection.BaseIndex + IDValue * 3 + 2]);
+	}
+
+	inline FVector3d GetVertex(int32 IDValue) const
+	{
+		const FVector3f& Position = Mesh->StaticVertexBuffers.PositionVertexBuffer.VertexPosition(IDValue);
+		return FVector3d(BuildScale.X * (double)Position.X, BuildScale.Y * (double)Position.Y, BuildScale.Z * (double)Position.Z);
+	}
+
+	inline void GetTriVertices(int32 IDValue, FVector3d& V0, FVector3d& V1, FVector3d& V2) const
+	{
+		FIndex3i VtxIndices = GetTriangle(IDValue);
+
+		V0 = GetVertex(VtxIndices.A);
+		V1 = GetVertex(VtxIndices.B);
+		V2 = GetVertex(VtxIndices.C);
+	}
+
+	template<typename VectorType>
+	inline void GetTriVertices(int32 IDValue, VectorType& V0, VectorType& V1, VectorType& V2) const
+	{
+		FIndex3i VtxIndices = GetTriangle(IDValue);
+
+		V0 = GetVertex(VtxIndices.A);
+		V1 = GetVertex(VtxIndices.B);
+		V2 = GetVertex(VtxIndices.C);
+	}
+
+	inline bool HasNormals() const
+	{
+		return Mesh && Mesh->StaticVertexBuffers.StaticMeshVertexBuffer.GetAllowCPUAccess() && Mesh->StaticVertexBuffers.StaticMeshVertexBuffer.GetTangentData();
+	}
+	inline bool IsNormal(int32 NID) const
+	{
+		return NID >= 0 && NID < NormalCount();
+	}
+	inline int32 MaxNormalID() const
+	{
+		return NormalCount();
+	}
+	inline int32 NormalCount() const
+	{
+		return HasNormals() ? VertexCount() : 0;
+	}
+
+	inline FVector3f GetNormal(int32 IDValue) const
+	{
+		const FVector4f& Normal = Mesh->StaticVertexBuffers.StaticMeshVertexBuffer.VertexTangentZ(IDValue);
+		return (!bScaleNormals) ? FVector3f(Normal.X, Normal.Y, Normal.Z) :
+			UE::Geometry::Normalized(FVector3f(Normal.X * InvBuildScale.X, Normal.Y * InvBuildScale.Y, Normal.Z * InvBuildScale.Z));
+	}
+
+	/** Get Normals for a given Triangle */
+	template<typename VectorType>
+	inline void GetTriNormals(int32 TriId, VectorType& N0, VectorType& N1, VectorType& N2)
+	{
+		FIndex3i VtxIndices = GetTriangle(TriId);
+
+		N0 = GetNormal(VtxIndices.A);
+		N1 = GetNormal(VtxIndices.B);
+		N2 = GetNormal(VtxIndices.C);
+	}
+
+
+	inline FVector3f GetTangentX(int32 IDValue) const
+	{
+		const FVector4f& TangentX = Mesh->StaticVertexBuffers.StaticMeshVertexBuffer.VertexTangentX(IDValue);
+		return (!bScaleNormals) ? FVector3f(TangentX.X, TangentX.Y, TangentX.Z) :
+			UE::Geometry::Normalized(FVector3f(TangentX.X * BuildScale.X, TangentX.Y * BuildScale.Y, TangentX.Z * BuildScale.Z));
+	}
+
+	/** Get Tangent X for a given Triangle */
+	template<typename VectorType>
+	inline void GetTriTangentsX(int32 TriId, VectorType& T0, VectorType& T1, VectorType& T2)
+	{
+		FIndex3i VtxIndices = GetTriangle(TriId);
+		T0 = GetTangentX(VtxIndices.A);
+		T1 = GetTangentX(VtxIndices.B);
+		T2 = GetTangentX(VtxIndices.C);
+	}
+
+
+	inline FVector3f GetTangentY(int32 IDValue) const
+	{
+		const FVector4f& TangentY = Mesh->StaticVertexBuffers.StaticMeshVertexBuffer.VertexTangentY(IDValue);
+		return (!bScaleNormals) ? FVector3f(TangentY.X, TangentY.Y, TangentY.Z) :
+			UE::Geometry::Normalized(FVector3f(TangentY.X * BuildScale.X, TangentY.Y * BuildScale.Y, TangentY.Z * BuildScale.Z));
+	}
+
+	/** Get Tangent Y for a given Triangle */
+	template<typename VectorType>
+	inline void GetTriTangentsY(int32 TriId, VectorType& T0, VectorType& T1, VectorType& T2)
+	{
+		FIndex3i VtxIndices = GetTriangle(TriId);
+		T0 = GetTangentY(VtxIndices.A);
+		T1 = GetTangentY(VtxIndices.B);
+		T2 = GetTangentY(VtxIndices.C);
+	}
+
+
+	inline bool HasUVs(const int32 UVLayer = 0) const
+	{
+		return Mesh && Mesh->StaticVertexBuffers.StaticMeshVertexBuffer.GetAllowCPUAccess() && Mesh->StaticVertexBuffers.StaticMeshVertexBuffer.GetTexCoordData() && (int32)Mesh->StaticVertexBuffers.StaticMeshVertexBuffer.GetNumTexCoords() > UVLayer && UVLayer >= 0;
+	}
+	inline int32 NumUVLayers() const
+	{
+		return (Mesh && Mesh->StaticVertexBuffers.StaticMeshVertexBuffer.GetAllowCPUAccess() && Mesh->StaticVertexBuffers.StaticMeshVertexBuffer.GetTexCoordData()) ?
+			(int32)Mesh->StaticVertexBuffers.StaticMeshVertexBuffer.GetNumTexCoords() : 0;
+	}
+	inline bool IsUV(const int32 UVId) const
+	{
+		return HasUVs() && UVId >= 0 && UVId < UVCount();
+	}
+	inline int32 MaxUVID() const
+	{
+		return UVCount();
+	}
+	inline int32 UVCount() const
+	{
+		return HasUVs() ? Mesh->GetNumVertices() : 0;
+	}
+
+	/** Get UV by VertexInstanceID for a given UVLayer */
+	inline FVector2f GetUV(const int32 IDValue, const int32 UVLayer) const
+	{
+		return Mesh->StaticVertexBuffers.StaticMeshVertexBuffer.GetVertexUV(IDValue, UVLayer);
+	}
+
+	/** Get UVs for a given UVLayer and Triangle */
+	template<typename VectorType>
+	inline void GetTriUVs(const int32 TriId, const int32 UVLayer, VectorType& UV0, VectorType& UV1, VectorType& UV2)
+	{
+		FIndex3i VtxIndices = GetTriangle(TriId);
+
+		UV0 = GetUV(VtxIndices.A, UVLayer);
+		UV1 = GetUV(VtxIndices.B, UVLayer);
+		UV2 = GetUV(VtxIndices.C, UVLayer);
+	}
+
+
+	inline bool HasColors() const
+	{
+		return Mesh && Mesh->StaticVertexBuffers.ColorVertexBuffer.GetAllowCPUAccess();
+	}
+	inline bool IsColor(int32 ColorIndex) const
+	{
+		return ColorIndex >= 0 && ColorIndex < ColorCount();
+	}
+	inline int32 MaxColorID() const
+	{
+		return ColorCount();
+	}
+	inline int32 ColorCount() const
+	{
+		return HasColors() ? VertexCount() : 0;
+	}
+
+	inline FColor GetColor(int32 IDValue) const
+	{
+		return Mesh->StaticVertexBuffers.ColorVertexBuffer.VertexColor(IDValue);
+	}
+
+	/** Get Colors for a given Triangle */
+	inline void GetTriColors(int32 TriId, FColor& C0, FColor& C1, FColor& C2)
+	{
+		FIndex3i VtxIndices = GetTriangle(TriId);
+		C0 = GetColor(VtxIndices.A);
+		C1 = GetColor(VtxIndices.B);
+		C2 = GetColor(VtxIndices.C);
+	}
+
+	inline bool HasSkinWeights()
+	{	
+		const bool bSkinWeightsValid = Mesh && 
+									   Mesh->SkinWeightVertexBuffer.GetDataVertexBuffer() && 
+									   Mesh->SkinWeightVertexBuffer.GetDataVertexBuffer()->IsInitialized() && 
+									   Mesh->SkinWeightVertexBuffer.GetDataVertexBuffer()->GetNumVertices() > 0;
+		
+		return bSkinWeightsValid;
+	}
+	
+	inline FSkinWeightInfo GetSkinWeightInfo(int32 VID)
+	{
+		return SkinWeights[VID];
+	}
+};
\ No newline at end of file
diff --git a/Engine/Source/Runtime/MeshConversionEngineTypes/Private/SkeletalMeshLODRenderDataToDynamicMesh.cpp b/Engine/Source/Runtime/MeshConversionEngineTypes/Private/SkeletalMeshLODRenderDataToDynamicMesh.cpp
new file mode 100644
--- /dev/null
+++ b/Engine/Source/Runtime/MeshConversionEngineTypes/Private/SkeletalMeshLODRenderDataToDynamicMesh.cpp
@@ -0,0 +1,285 @@
+// Copyright Epic Games, Inc. All Rights Reserved.
+
+#include "SkeletalMeshLODRenderDataToDynamicMesh.h"
+#include "SkeletalMeshLODRenderDataMeshAdapter.h"
+#include "DynamicMesh/DynamicMeshAttributeSet.h"
+#include "DynamicMesh/DynamicVertexSkinWeightsAttribute.h"
+#include "Rendering/SkeletalMeshLODRenderData.h"
+#include "SkeletalMeshAttributes.h"
+#include "DynamicMesh/DynamicBoneAttribute.h"
+#include "ReferenceSkeleton.h"
+
+using namespace UE::Geometry;
+
+bool FSkeletalMeshLODRenderDataToDynamicMesh::Convert(
+	const FSkeletalMeshLODRenderData* SkeletalMeshResources,
+	const FReferenceSkeleton& RefSkeleton,
+	const ConversionOptions& Options,
+	FDynamicMesh3& OutputMesh,
+	bool bHasVertexColors,
+	TFunctionRef<FColor(int32)> GetVertexColorFromLODVertexIndex)
+{
+	if (!ensureMsgf(SkeletalMeshResources && SkeletalMeshResources->StaticVertexBuffers.StaticMeshVertexBuffer.GetAllowCPUAccess(), TEXT("bAllowCPUAccess must be set to true for SkeletalMesh before calling FSkeletalMeshLODRenderDataToDynamicMesh::Convert(), otherwise the mesh geometry data isn't accessible!")))
+	{
+		return false;
+	}
+
+	FSkeletalMeshLODRenderDataMeshAdapter Adapter(SkeletalMeshResources);
+
+	Adapter.SetBuildScale(Options.BuildScale, false);
+
+	OutputMesh = FDynamicMesh3();
+	OutputMesh.EnableTriangleGroups();
+	if (Options.bWantNormals || Options.bWantTangents || Options.bWantUVs || Options.bWantVertexColors || Options.bWantMaterialIDs)
+	{
+		OutputMesh.EnableAttributes();
+	}
+
+	// map from FSkeletalMeshLODRenderData triangle ID to DynamicMesh TID
+	TArray<int32> ToDstTriID;
+
+	/**
+	* map from DynamicMesh vertex Id to FSkeletalMeshLODRenderData vertex id.
+	* NB: due to vertex splitting, multiple DynamicMesh vertex ids
+	* may map to the same FSkeletalMeshLODRenderData vertex id.
+	*  ( a vertex split is a result of reconciling non-manifold MeshDescription vertex )
+	*/
+	TArray<int32> ToSrcVID;
+
+	// Copy vertices. LODMesh is dense so this should be 1-1
+	int32 SrcVertexCount = Adapter.VertexCount();
+	ToSrcVID.SetNumUninitialized(SrcVertexCount);
+	for (int32 SrcVertID = 0; SrcVertID < SrcVertexCount; ++SrcVertID)
+	{
+		FVector3d Position = Adapter.GetVertex(SrcVertID);
+		int32 DstVertID = OutputMesh.AppendVertex(Position);
+		ToSrcVID[DstVertID] = SrcVertID;
+
+		if (DstVertID != SrcVertID)
+		{
+			// should only happen in the source mesh is missing vertices. 
+			OutputMesh.Clear();
+			ensure(false);
+			return false;
+		}
+	}
+
+	// Copy triangles. LODMesh is dense so this should be 1-1 unless there is a duplicate tri or non-manifold edge 
+	int32 SrcTriangleCount = Adapter.TriangleCount();
+	ToDstTriID.Init(FDynamicMesh3::InvalidID, SrcTriangleCount);
+	for (int32 SrcTriID = 0; SrcTriID < SrcTriangleCount; ++SrcTriID)
+	{
+		FIndex3i Tri = Adapter.GetTriangle(SrcTriID);
+		int32 DstTriID = OutputMesh.AppendTriangle(Tri.A, Tri.B, Tri.C);
+
+		if (DstTriID == FDynamicMesh3::DuplicateTriangleID || DstTriID == FDynamicMesh3::InvalidID)
+		{
+			continue;
+		}
+
+		// split verts on the non-manifold edge(s)
+		if (DstTriID == FDynamicMesh3::NonManifoldID)
+		{
+			int e01 = OutputMesh.FindEdge(Tri[0], Tri[1]);
+			int e12 = OutputMesh.FindEdge(Tri[1], Tri[2]);
+			int e20 = OutputMesh.FindEdge(Tri[2], Tri[0]);
+
+			// determine which verts need to be duplicated
+			bool bToSplit[3] = { false, false, false };
+			if (e01 != FDynamicMesh3::InvalidID && OutputMesh.IsBoundaryEdge(e01) == false)
+			{
+				bToSplit[0] = true;
+				bToSplit[1] = true;
+			}
+			if (e12 != FDynamicMesh3::InvalidID && OutputMesh.IsBoundaryEdge(e12) == false)
+			{
+				bToSplit[1] = true;
+				bToSplit[2] = true;
+			}
+			if (e20 != FDynamicMesh3::InvalidID && OutputMesh.IsBoundaryEdge(e20) == false)
+			{
+				bToSplit[2] = true;
+				bToSplit[0] = true;
+			}
+			for (int32 i = 0; i < 3; ++i)
+			{
+				if (bToSplit[i])
+				{
+					const int32 SrcVID = Tri[i];
+					const FVector3d Position = Adapter.GetVertex(SrcVID);
+					const int32 NewDstVertIdx = OutputMesh.AppendVertex(Position);
+					Tri[i] = NewDstVertIdx;
+					ToSrcVID.SetNumUninitialized(NewDstVertIdx + 1);
+					ToSrcVID[NewDstVertIdx] = SrcVID;
+				}
+			}
+
+			DstTriID = OutputMesh.AppendTriangle(Tri[0], Tri[1], Tri[2]);
+		}
+
+		ToDstTriID[SrcTriID] = DstTriID;
+
+	}
+
+	// transfer sections to PolyGroups and MaterialIDs
+	FDynamicMeshMaterialAttribute* MaterialIDs = nullptr;
+	if (Options.bWantMaterialIDs)
+	{
+		OutputMesh.Attributes()->EnableMaterialID();
+		MaterialIDs = OutputMesh.Attributes()->GetMaterialID();
+	}
+	for (int32 SectionIdx = 0; SectionIdx < SkeletalMeshResources->RenderSections.Num(); ++SectionIdx)
+	{
+		const FSkelMeshRenderSection& Section = SkeletalMeshResources->RenderSections[SectionIdx];
+		for (uint32 TriIdx = 0; TriIdx < Section.NumTriangles; ++TriIdx)
+		{
+			int32 SrcTriangleID = (int32)(Section.BaseIndex / 3 + TriIdx);
+			int32 DstTriID = ToDstTriID[SrcTriangleID];
+			if (DstTriID != FDynamicMesh3::InvalidID)
+			{
+				OutputMesh.SetTriangleGroup(DstTriID, SectionIdx);
+				if (MaterialIDs != nullptr)
+				{
+					int32 MaterialIndex = (int32)Section.MaterialIndex;
+					MaterialIDs->SetValue(DstTriID, MaterialIndex);
+				}
+			}
+		}
+	}
+
+	// Helper to copy a per-vertex attribute to an overlay
+	auto CopyPerVertexOverlay = [SrcTriangleCount, &ToSrcVID, &ToDstTriID, &OutputMesh](auto* Overlay, auto GetElement)
+		{
+			const int32 DstVertexCount = ToSrcVID.Num();
+			for (int32 DstVertID = 0; DstVertID < DstVertexCount; ++DstVertID)
+			{
+				const int32 SrcVID = ToSrcVID[DstVertID];
+				auto Element = GetElement(SrcVID);
+				int32 ElemID = Overlay->AppendElement(Element);
+				check(ElemID == DstVertID);
+			}
+
+			for (int32 SrcTriID = 0; SrcTriID < SrcTriangleCount; ++SrcTriID)
+			{
+				const int32 DstTriID = ToDstTriID[SrcTriID];
+				if (DstTriID != FDynamicMesh3::InvalidID)
+				{
+					FIndex3i Tri = OutputMesh.GetTriangle(DstTriID);
+					Overlay->SetTriangle(DstTriID, FIndex3i(Tri.A, Tri.B, Tri.C));
+				}
+			}
+		};
+
+	// copy overlay normals
+	if (Adapter.HasNormals() && Options.bWantNormals)
+	{
+		FDynamicMeshNormalOverlay* Normals = OutputMesh.Attributes()->PrimaryNormals();
+		if (Normals != nullptr)
+		{
+			CopyPerVertexOverlay(Normals, [&](int32 SrcVID)->FVector3f { return Adapter.GetNormal(SrcVID); });
+		}
+	}
+
+	// copy overlay tangents
+	if (Adapter.HasNormals() && Options.bWantTangents)
+	{
+		OutputMesh.Attributes()->EnableTangents();
+		FDynamicMeshNormalOverlay* TangentsX = OutputMesh.Attributes()->PrimaryTangents();
+		if (TangentsX != nullptr)
+		{
+			CopyPerVertexOverlay(TangentsX, [&](int32 SrcVID)->FVector3f { return Adapter.GetTangentX(SrcVID); });
+		}
+
+		FDynamicMeshNormalOverlay* TangentsY = OutputMesh.Attributes()->PrimaryBiTangents();
+		if (TangentsY != nullptr)
+		{
+			CopyPerVertexOverlay(TangentsY, [&](int32 SrcVID)->FVector3f { return Adapter.GetTangentY(SrcVID); });
+		}
+	}
+
+	// copy UV layers
+	if (Adapter.HasUVs() && Options.bWantUVs)
+	{
+		int32 NumUVLayers = Adapter.NumUVLayers();
+		if (NumUVLayers > 0)
+		{
+			OutputMesh.Attributes()->SetNumUVLayers(NumUVLayers);
+			for (int32 UVLayerIndex = 0; UVLayerIndex < NumUVLayers; ++UVLayerIndex)
+			{
+				FDynamicMeshUVOverlay* UVOverlay = OutputMesh.Attributes()->GetUVLayer(UVLayerIndex);
+				CopyPerVertexOverlay(UVOverlay, [&](int32 SrcVID)->FVector2f { return Adapter.GetUV(SrcVID, UVLayerIndex); });
+			}
+		}
+	}
+
+	// copy overlay colors
+	if (bHasVertexColors && Options.bWantVertexColors)
+	{
+		OutputMesh.Attributes()->EnablePrimaryColors();
+		FDynamicMeshColorOverlay* Colors = OutputMesh.Attributes()->PrimaryColors();
+
+		CopyPerVertexOverlay(Colors, [&](int32 SrcVID)->FVector4f { return GetVertexColorFromLODVertexIndex(SrcVID).ReinterpretAsLinear(); });
+	}
+
+	// copy skin weights
+	if (Adapter.HasSkinWeights() && Options.bWantSkinWeights)
+	{
+		auto CopyPerVertexSkinWeights = [&ToSrcVID, &OutputMesh, &Adapter](FDynamicMeshVertexSkinWeightsAttribute* Overlay)
+		{
+			const int32 DstVertexCount = ToSrcVID.Num();
+			for (int32 DstVertID = 0; DstVertID < DstVertexCount; ++DstVertID)
+			{
+				const int32 SrcVID = ToSrcVID[DstVertID];
+				const FSkinWeightInfo WeightInfo = Adapter.GetSkinWeightInfo(SrcVID);;
+				const UE::AnimationCore::FBoneWeights Weights = UE::AnimationCore::FBoneWeights::Create(WeightInfo.InfluenceBones, WeightInfo.InfluenceWeights);
+				Overlay->SetValue(DstVertID, Weights);
+			}
+		};
+
+		
+		FDynamicMeshVertexSkinWeightsAttribute* SkinAttribute = new FDynamicMeshVertexSkinWeightsAttribute(&OutputMesh);
+
+		CopyPerVertexSkinWeights(SkinAttribute);
+		
+		SkinAttribute->SetName(FSkeletalMeshAttributes::DefaultSkinWeightProfileName);
+
+		OutputMesh.Attributes()->AttachSkinWeightsAttribute(FSkeletalMeshAttributes::DefaultSkinWeightProfileName, SkinAttribute);
+
+		// populate bones attribute
+		const int32 NumBones = RefSkeleton.GetRawBoneNum();
+		if (NumBones)
+		{
+			OutputMesh.Attributes()->EnableBones(NumBones);
+			FDynamicMeshBoneNameAttribute* BoneNameAttrib = OutputMesh.Attributes()->GetBoneNames();
+			FDynamicMeshBoneParentIndexAttribute* BoneParentIndices = OutputMesh.Attributes()->GetBoneParentIndices();
+			FDynamicMeshBonePoseAttribute* BonePoses = OutputMesh.Attributes()->GetBonePoses();
+			FDynamicMeshBoneColorAttribute* BoneColors = OutputMesh.Attributes()->GetBoneColors();
+
+			for (int32 BoneIdx = 0; BoneIdx < NumBones; ++BoneIdx)
+			{
+				BoneNameAttrib->SetValue(BoneIdx, RefSkeleton.GetRawRefBoneInfo()[BoneIdx].Name);
+				BoneParentIndices->SetValue(BoneIdx, RefSkeleton.GetRawRefBoneInfo()[BoneIdx].ParentIndex);
+				BonePoses->SetValue(BoneIdx, RefSkeleton.GetRawRefBonePose()[BoneIdx]);
+				BoneColors->SetValue(BoneIdx, FVector4f::One());
+			}
+		}
+	}
+
+
+	return true;
+}
+
+
+bool FSkeletalMeshLODRenderDataToDynamicMesh::Convert(
+	const FSkeletalMeshLODRenderData* SkeletalMeshResources,
+	const FReferenceSkeleton& RefSkeleton,
+	const ConversionOptions& Options,
+	FDynamicMesh3& OutputMesh)
+{
+	bool bHasVertexColors = SkeletalMeshResources && SkeletalMeshResources->StaticVertexBuffers.ColorVertexBuffer.GetAllowCPUAccess();
+	return Convert(SkeletalMeshResources, RefSkeleton, Options, OutputMesh, bHasVertexColors,
+		[SkeletalMeshResources](int32 VID)
+		{
+			return SkeletalMeshResources->StaticVertexBuffers.ColorVertexBuffer.VertexColor(VID);
+		});
+}
diff --git a/Engine/Source/Runtime/MeshConversionEngineTypes/Public/SkeletalMeshLODRenderDataToDynamicMesh.h b/Engine/Source/Runtime/MeshConversionEngineTypes/Public/SkeletalMeshLODRenderDataToDynamicMesh.h
new file mode 100644
--- /dev/null
+++ b/Engine/Source/Runtime/MeshConversionEngineTypes/Public/SkeletalMeshLODRenderDataToDynamicMesh.h
@@ -0,0 +1,62 @@
+// Copyright Epic Games, Inc. All Rights Reserved.
+
+#pragma once
+
+#include "CoreMinimal.h"
+#include "DynamicMesh/DynamicMesh3.h"
+
+class FSkeletalMeshLODRenderData;
+struct FReferenceSkeleton;
+
+namespace UE
+{
+namespace Geometry
+{
+/**
+ * FSkeletalMeshLODRenderDataToDynamicMesh can be used to create a FDynamicMesh3 from
+ * a FSkeletalMeshLODRenderData, which is the variant of the mesh in a SkeletalMesh Asset used for rendering
+ * (and which is available at runtime). The FSkeletalMeshLODRenderData has vertices duplicated at any
+ * split UV/normal/tangent/color, ie in the overlays there will be a unique overlay element for each
+ * base mesh vertex.
+ * 
+ * TODO: The code mainly follows the logic of the StaticMeshLODResourcesToDynamicMesh.h with the addition of the 
+ * skinning weights and bones. Both should be refactored to use ToDynamicMesh interface instead.
+ */
+
+
+ class FSkeletalMeshLODRenderDataToDynamicMesh
+ {
+ public:
+
+	// TODO: more options?
+	 struct ConversionOptions
+	 {
+		 bool bWantNormals = true;
+		 bool bWantTangents = true;
+		 bool bWantUVs = true;
+		 bool bWantVertexColors = true;
+		 bool bWantMaterialIDs = true;
+		 bool bWantSkinWeights = true;
+
+		 // mesh vertex positions are multiplied by the build scale
+		 FVector3d BuildScale = FVector3d::One();
+	 };
+
+	// @param bHasVertexColors					Whether to add vertex colors to the output (if available)
+	// @param GetVertexColorFromLODVertexIndex	Function for getting the vertex color of a given vertex index
+	 static MESHCONVERSIONENGINETYPES_API bool Convert(
+		 const FSkeletalMeshLODRenderData* SkeletalMeshResources,
+		 const FReferenceSkeleton& RefSkeleton,
+		 const ConversionOptions& Options,
+		 FDynamicMesh3& OutputMesh,
+		 bool bHasVertexColors,
+		 TFunctionRef<FColor(int32)> GetVertexColorFromLODVertexIndex);
+
+	 static MESHCONVERSIONENGINETYPES_API bool Convert(
+		 const FSkeletalMeshLODRenderData* SkeletalMeshResources,
+		 const FReferenceSkeleton& RefSkeleton,
+		 const ConversionOptions& Options,
+		 FDynamicMesh3& OutputMesh);
+ };
+}
+}
\ No newline at end of file
