Add a geometry collection to dynamic mesh conversion struct to the MeshConversion module, to support dynamic mesh-based updates of geometry collections in more places. This roughly follows the approach of the PlanarCuts module (which has a specialized/private version of this, for fracture only), but puts the data into the dynamic mesh using standard overlays instead of custom attributes.

As an initial use of this new conversion, add a dataflow node to run a geometry script processor on selected bones of a geometry collection.

diff --git a/Engine/Plugins/Experimental/GeometryCollectionPlugin/Source/GeometryCollectionNodes/Private/Dataflow/GeometryCollectionMeshNodes.cpp b/Engine/Plugins/Experimental/GeometryCollectionPlugin/Source/GeometryCollectionNodes/Private/Dataflow/GeometryCollectionMeshNodes.cpp
--- a/Engine/Plugins/Experimental/GeometryCollectionPlugin/Source/GeometryCollectionNodes/Private/Dataflow/GeometryCollectionMeshNodes.cpp
+++ b/Engine/Plugins/Experimental/GeometryCollectionPlugin/Source/GeometryCollectionNodes/Private/Dataflow/GeometryCollectionMeshNodes.cpp
@@ -10,6 +10,7 @@
 #include "GeometryCollection/GeometryCollection.h"
 #include "GeometryCollection/GeometryCollectionEngineUtility.h"
 #include "GeometryCollection/GeometryCollectionEngineConversion.h"
+#include "GeometryCollection/Facades/CollectionTransformSelectionFacade.h"
 #include "Logging/LogMacros.h"
 #include "Templates/SharedPointer.h"
 #include "UObject/UnrealTypePrivate.h"
@@ -19,6 +20,8 @@
 #include "DynamicMeshEditor.h"
 #include "Operations/MeshBoolean.h"
 
+#include "GeometryCollectionToDynamicMesh.h"
+
 #include "EngineGlobals.h"
 #include "GeometryCollection/GeometryCollectionAlgo.h"
 #include "GeometryCollection/GeometryCollectionClusteringUtility.h"
@@ -52,6 +55,7 @@ namespace UE::Dataflow
 		DATAFLOW_NODE_REGISTER_CREATION_FACTORY(FMeshCopyToPointsDataflowNode);
 		DATAFLOW_NODE_REGISTER_CREATION_FACTORY(FGetMeshDataDataflowNode);
 		DATAFLOW_NODE_REGISTER_CREATION_FACTORY(FApplyMeshProcessorToMeshDataflowNode);
+		DATAFLOW_NODE_REGISTER_CREATION_FACTORY(FApplyMeshProcessorToGeometryCollectionDataflowNode);
 
 
 		// Mesh
@@ -494,3 +498,90 @@ void FApplyMeshProcessorToMeshDataflowNode::Evaluate(UE::Dataflow::FContext& Con
 		}
 	}
 }
+
+
+void FApplyMeshProcessorToGeometryCollectionDataflowNode::Evaluate(UE::Dataflow::FContext& Context, const FDataflowOutput* Out) const
+{
+	if (Out->IsA(&Collection) ||
+		Out->IsA(&TransformSelection))
+	{
+		if (!MeshProcessorInstance)
+		{
+			SafeForwardInput(Context, &Collection, &Collection);
+			return;
+		}
+
+		FDataflowTransformSelection InTransformSelection = GetValue(Context, &TransformSelection);
+
+		//
+		// If not connected select everything by default
+		//
+		if (!IsConnected(&TransformSelection))
+		{
+			const FManagedArrayCollection& InCollection = GetValue(Context, &Collection);
+
+			GeometryCollection::Facades::FCollectionTransformSelectionFacade TransformSelectionFacade(InCollection);
+			const TArray<int32>& SelectionArr = TransformSelectionFacade.SelectAll();
+
+			FDataflowTransformSelection NewTransformSelection;
+			NewTransformSelection.Initialize(InCollection.NumElements(FGeometryCollection::TransformGroup), false);
+			NewTransformSelection.SetFromArray(SelectionArr);
+
+			InTransformSelection = NewTransformSelection;
+		}
+
+		if (InTransformSelection.AnySelected())
+		{
+			const FManagedArrayCollection& InCollection = GetValue(Context, &Collection);
+
+			UE::Geometry::FGeometryCollectionToDynamicMeshes CollectionToMeshes;
+			UE::Geometry::FGeometryCollectionToDynamicMeshes::FToMeshOptions ToMeshOptions;
+			ToMeshOptions.bWeldVertices = bWeldVertices;
+			ToMeshOptions.bSaveIsolatedVertices = bPreserveIsolatedVertices;
+			if (CollectionToMeshes.InitFromTransformSelection(InCollection, InTransformSelection.AsArray(), ToMeshOptions)
+				&& !CollectionToMeshes.Meshes.IsEmpty())
+			{
+				// Temporarily create a UDynamicMesh as a container to hold the meshes we pass to BP
+				TObjectPtr<UDynamicMesh> NewMesh = NewObject<UDynamicMesh>();
+
+				bool bAnySuccess = false;
+				for (UE::Geometry::FGeometryCollectionToDynamicMeshes::FMeshInfo& MeshInfo : CollectionToMeshes.Meshes)
+				{
+					NewMesh->SetMesh(MoveTemp(*MeshInfo.Mesh));
+
+					bool bFailed = false;
+					MeshProcessorInstance->ProcessDynamicMesh(NewMesh, bFailed);
+					if (!bFailed) // on success, move the mesh back
+					{
+						bAnySuccess = true;
+						NewMesh->EditMesh([&MeshInfo](UE::Geometry::FDynamicMesh3& Mesh)
+						{
+							*MeshInfo.Mesh = MoveTemp(Mesh);
+						}, EDynamicMeshChangeType::GeneralEdit, EDynamicMeshAttributeChangeFlags::Unknown, true);
+					}
+					else // on failure, clear the mesh so it won't be written back
+					{
+						MeshInfo.TransformIndex = -1;
+						MeshInfo.Mesh = nullptr;
+					}
+				}
+
+				if (bAnySuccess)
+				{
+					if (TUniquePtr<FGeometryCollection> GeomCollection = TUniquePtr<FGeometryCollection>(InCollection.NewCopy<FGeometryCollection>()))
+					{
+						UE::Geometry::FGeometryCollectionToDynamicMeshes::FToCollectionOptions ToCollectionOptions;
+						ToCollectionOptions.bDefaultFaceInternal = false;
+						ToCollectionOptions.bDefaultFaceVisible = true;
+						CollectionToMeshes.UpdateGeometryCollection(*GeomCollection, ToCollectionOptions);
+						SetValue<const FManagedArrayCollection&>(Context, *GeomCollection, &Collection);
+						return;
+					}
+				}
+			}
+		}
+
+		SafeForwardInput(Context, &Collection, &Collection);
+	}
+}
+
diff --git a/Engine/Plugins/Experimental/GeometryCollectionPlugin/Source/GeometryCollectionNodes/Public/Dataflow/GeometryCollectionMeshNodes.h b/Engine/Plugins/Experimental/GeometryCollectionPlugin/Source/GeometryCollectionNodes/Public/Dataflow/GeometryCollectionMeshNodes.h
--- a/Engine/Plugins/Experimental/GeometryCollectionPlugin/Source/GeometryCollectionNodes/Public/Dataflow/GeometryCollectionMeshNodes.h
+++ b/Engine/Plugins/Experimental/GeometryCollectionPlugin/Source/GeometryCollectionNodes/Public/Dataflow/GeometryCollectionMeshNodes.h
@@ -134,6 +134,46 @@ public:
 	virtual void Evaluate(UE::Dataflow::FContext& Context, const FDataflowOutput* Out) const override;
 };
 
+/**
+ * Apply a Geometry Script Mesh Processors to the geometry of selected transforms in a geometry collection
+ */
+USTRUCT(meta = (DataflowGeometryCollection))
+struct FApplyMeshProcessorToGeometryCollectionDataflowNode final : public FMeshProcessorDataflowNodeBase
+{
+	GENERATED_USTRUCT_BODY()
+	DATAFLOW_NODE_DEFINE_INTERNAL(FApplyMeshProcessorToGeometryCollectionDataflowNode, "ApplyGeometryScriptToCollection", "Mesh|Utilities", "")
+	DATAFLOW_NODE_RENDER_TYPE("SurfaceRender", FGeometryCollection::StaticType(), "Collection")
+
+public:
+
+	/** Input/Output mesh */
+	UPROPERTY(meta = (DataflowInput, DataflowOutput, DataflowPassthrough = "Collection", DataflowIntrinsic))
+	FManagedArrayCollection Collection;
+
+	/** Selected bones will have geometry script processing applied (if they have geometry). If not connected, all geometry will be processed. */
+	UPROPERTY(meta = (DataflowInput))
+	FDataflowTransformSelection TransformSelection;
+
+	/** Whether the processed mesh will have edges at normal/UV/color seams welded so they are treated as one edge during processing. */
+	UPROPERTY(EditAnywhere, Category = Options)
+	bool bWeldVertices = true;
+
+	/** Whether to preserve isolated vertices which aren't used by any triangles. */
+	UPROPERTY(EditAnywhere, Category = Options)
+	bool bPreserveIsolatedVertices = true;
+
+	FApplyMeshProcessorToGeometryCollectionDataflowNode(const UE::Dataflow::FNodeParameters& InParam, FGuid InGuid = FGuid::NewGuid())
+		: FMeshProcessorDataflowNodeBase(InParam, InGuid)
+	{
+		RegisterInputConnection(&Collection);
+
+		RegisterInputConnection(&TransformSelection);
+		RegisterOutputConnection(&Collection, &Collection);
+	}
+
+	virtual void Evaluate(UE::Dataflow::FContext& Context, const FDataflowOutput* Out) const override;
+};
+
 /**
  *
  * Converts a BoundingBox into a DynamicMesh
diff --git a/Engine/Source/Runtime/MeshConversion/MeshConversion.Build.cs b/Engine/Source/Runtime/MeshConversion/MeshConversion.Build.cs
--- a/Engine/Source/Runtime/MeshConversion/MeshConversion.Build.cs
+++ b/Engine/Source/Runtime/MeshConversion/MeshConversion.Build.cs
@@ -16,6 +16,7 @@ public class MeshConversion : ModuleRules
         PublicDependencyModuleNames.AddRange(
 			new string[] {
                 "Core",
+				"Chaos", // for geometry collection
                 "MeshDescription",
 				"StaticMeshDescription",
 				"SkeletalMeshDescription",
diff --git a/Engine/Source/Runtime/MeshConversion/Private/GeometryCollectionToDynamicMesh.cpp b/Engine/Source/Runtime/MeshConversion/Private/GeometryCollectionToDynamicMesh.cpp
new file mode 100644
--- /dev/null
+++ b/Engine/Source/Runtime/MeshConversion/Private/GeometryCollectionToDynamicMesh.cpp
@@ -0,0 +1,615 @@
+// Copyright Epic Games, Inc. All Rights Reserved.
+
+#include "GeometryCollectionToDynamicMesh.h"
+
+#include "Containers/ArrayView.h"
+
+#include "DynamicMesh/DynamicMesh3.h"
+#include "DynamicMesh/DynamicMeshAttributeSet.h"
+#include "DynamicMesh/Operations/MergeCoincidentMeshEdges.h"
+#include "DynamicMeshEditor.h"
+
+#include "GeometryCollection/GeometryCollection.h"
+#include "GeometryCollection/GeometryCollectionAlgo.h"
+#include "GeometryCollection/TransformCollection.h"
+
+#include "GeometryCollection/Facades/CollectionMeshFacade.h"
+#include "GeometryCollection/Facades/CollectionTransformFacade.h"
+#include "GeometryCollection/Facades/CollectionUVFacade.h"
+
+namespace UE::Geometry
+{
+
+TConstArrayView<FTransform3f> FGeometryCollectionToDynamicMeshes::GetCollectionTransforms(const FManagedArrayCollection& Collection)
+{
+	const GeometryCollection::Facades::FCollectionTransformFacade TransformFacade(Collection);
+	if (TransformFacade.IsValid())
+	{
+		if (const TManagedArray<FTransform3f>* Transforms = TransformFacade.FindTransforms())
+		{
+			return TConstArrayView<FTransform3f>(Transforms->GetConstArray());
+		}
+	}
+	return TConstArrayView<FTransform3f>();
+}
+
+bool FGeometryCollectionToDynamicMeshes::InitHelper(const FManagedArrayCollection& Collection, bool bTransformInComponentSpace, TConstArrayView<FTransform3f> Transforms, bool bAllTransforms, TConstArrayView<int32> TransformIndices, const FToMeshOptions& Options)
+{
+	Meshes.Reset();
+
+	const GeometryCollection::Facades::FCollectionMeshFacade MeshFacade(Collection);
+	if (!MeshFacade.IsValid())
+	{
+		return false;
+	}
+	TManagedArrayAccessor<int32> ParentAttribute(Collection, FTransformCollection::ParentAttribute, FTransformCollection::TransformGroup);
+	if (!ParentAttribute.IsValid())
+	{
+		return false;
+	}
+	TManagedArrayAccessor<int32> SimTypeAttribute(Collection, FGeometryCollection::SimulationTypeAttribute, FTransformCollection::TransformGroup);
+
+	if (Transforms.IsEmpty())
+	{
+		return true; // successfully converted an empty collection?
+	}
+
+	const GeometryCollection::UV::FConstUVLayers UVLayers = GeometryCollection::UV::FindActiveUVLayers(Collection);
+
+	TConstArrayView<int32> UseTransformIndices = TransformIndices;
+	TArray<int32> LocalTransformIndices;
+	if (bAllTransforms)
+	{
+		LocalTransformIndices.SetNum(Transforms.Num());
+		for (int32 Idx = 0; Idx < Transforms.Num(); ++Idx)
+		{
+			LocalTransformIndices[Idx] = Idx;
+		}
+		UseTransformIndices = TConstArrayView<int32>(LocalTransformIndices);
+	}
+
+	for (int32 TransformIdx : UseTransformIndices)
+	{
+		int32 GeometryIdx = MeshFacade.TransformToGeometryIndexAttribute[TransformIdx];
+		if (GeometryIdx == INDEX_NONE)
+		{
+			// currently only converting transforms with associated geometry
+			continue;
+		}
+		if (SimTypeAttribute.IsValid() && SimTypeAttribute[TransformIdx] != FGeometryCollection::ESimulationTypes::FST_Rigid)
+		{
+			// if we have sim type info, also restrict to rigids -- cluster geometry is typically just there for legacy reasons
+			// (TODO: expose as option)
+			continue;
+		}
+
+		FTransformSRT3d CollectionToLocal;
+		if (bTransformInComponentSpace)
+		{
+			CollectionToLocal = FTransformSRT3d(FTransform(Transforms[TransformIdx]) * Options.Transform);
+		}
+		else
+		{
+			CollectionToLocal = FTransformSRT3d(GeometryCollectionAlgo::GlobalMatrix(Transforms, TArrayView<const int32>(ParentAttribute.Get().GetConstArray()), TransformIdx) * Options.Transform);
+		}
+
+		int32 AddedMeshIdx = Meshes.Emplace();
+		FMeshInfo& MeshData = Meshes[AddedMeshIdx];
+		MeshData.Transform = FTransform(CollectionToLocal);
+		MeshData.TransformIndex = TransformIdx;
+		MeshData.Mesh = MakeUnique<FDynamicMesh3>();
+		FDynamicMesh3& Mesh = *MeshData.Mesh;
+
+		Mesh.EnableAttributes();
+		Mesh.Attributes()->EnableMaterialID();
+		Mesh.Attributes()->EnablePrimaryColors();
+		Mesh.Attributes()->EnableTangents();
+		Mesh.Attributes()->SetNumUVLayers(UVLayers.Num());
+		bool bInvisiblePolygroup = Options.InvisibleFaces != EInvisibleFaceConversion::TagWithPolygroup;
+		int32 NumCustomPolygroupLayers = (int32)Options.bInternalFaceTagsAsPolygroups + (int32)bInvisiblePolygroup;
+		FDynamicMeshPolygroupAttribute* InternalFacePolygroups = nullptr;
+		FDynamicMeshPolygroupAttribute* VisibleFacePolygroups = nullptr;
+		Mesh.Attributes()->SetNumPolygroupLayers(NumCustomPolygroupLayers);
+		{
+			int32 CustomLayerIdx = 0;
+			if (Options.bInternalFaceTagsAsPolygroups)
+			{
+				InternalFacePolygroups = Mesh.Attributes()->GetPolygroupLayer(CustomLayerIdx);
+				InternalFacePolygroups->SetName(InternalFacePolyGroupName());
+				CustomLayerIdx++;
+			}
+			if (bInvisiblePolygroup)
+			{
+				VisibleFacePolygroups = Mesh.Attributes()->GetPolygroupLayer(CustomLayerIdx);
+				VisibleFacePolygroups->SetName(VisibleFacePolyGroupName());
+				CustomLayerIdx++;
+			}
+		}
+
+		const int32 VertexStart = MeshFacade.VertexStartAttribute[GeometryIdx];
+		const int32 VertexCount = MeshFacade.VertexCountAttribute[GeometryIdx];
+		const TArray<FVector3f>& VertexPositions = MeshFacade.VertexAttribute.Get().GetConstArray();
+		const TArray<FLinearColor>& VertexColors = MeshFacade.ColorAttribute.Get().GetConstArray();
+		const TArray<FVector3f>& VertexNormals = MeshFacade.NormalAttribute.Get().GetConstArray();
+		const TArray<FVector3f>& VertexTangents = MeshFacade.NormalAttribute.Get().GetConstArray();
+		const TArray<FVector3f>& VertexBitangents = MeshFacade.NormalAttribute.Get().GetConstArray();
+
+		for (int32 Idx = VertexStart, N = VertexStart + VertexCount; Idx < N; ++Idx)
+		{
+			FVector3d Position = CollectionToLocal.TransformPosition(FVector3d(VertexPositions[Idx]));
+			int32 VID = Mesh.AppendVertex(Position);
+			// add overlay elements -- by construction, will be 1:1 with vertices initially
+			const int32 ColorEID = Mesh.Attributes()->PrimaryColors()->AppendElement((FVector4f)VertexColors[Idx]);
+			checkSlow(ColorEID == VID);
+			const int32 NormalEID = Mesh.Attributes()->PrimaryNormals()->AppendElement(VertexNormals[Idx]);
+			checkSlow(NormalEID == VID);
+			const int32 TangentEID = Mesh.Attributes()->PrimaryTangents()->AppendElement(VertexTangents[Idx]);
+			checkSlow(TangentEID == VID);
+			const int32 BiTangentEID = Mesh.Attributes()->PrimaryBiTangents()->AppendElement(VertexBitangents[Idx]);
+			checkSlow(BiTangentEID == VID);
+			for (int32 UVLayer = 0; UVLayer < UVLayers.Num(); ++UVLayer)
+			{
+				const int32 UVElID = Mesh.Attributes()->GetUVLayer(UVLayer)->AppendElement(UVLayers[UVLayer][Idx]);
+				checkSlow(UVElID == VID);
+			}
+		}
+
+		// TODO: optionally enable non-manifold vertex mapping?
+
+		FIntVector VertexOffset(VertexStart, VertexStart, VertexStart);
+		const int32 FaceStart = MeshFacade.FaceStartAttribute[GeometryIdx];
+		const int32 FaceCount = MeshFacade.FaceCountAttribute[GeometryIdx];
+
+		const TArray<FIntVector>& FacesArray = MeshFacade.IndicesAttribute.Get().GetConstArray();
+		const TBitArray<>& VisibleArray = MeshFacade.VisibleAttribute.Get().GetConstArray();
+		const TBitArray<>& InternalArray = MeshFacade.InternalAttribute.Get().GetConstArray();
+		const TArray<int32>& MaterialIDs = MeshFacade.MaterialIDAttribute.Get().GetConstArray();
+
+		bool bSrcIsManifold = true;
+
+		for (int32 FaceIdx = FaceStart, N = FaceStart + FaceCount; FaceIdx < N; ++FaceIdx)
+		{
+			if (Options.InvisibleFaces == EInvisibleFaceConversion::Skip && !VisibleArray[FaceIdx])
+			{
+				continue;
+			}
+			FIndex3i AddTri = FacesArray[FaceIdx] - VertexOffset;
+			int32 TID = Mesh.AppendTriangle(AddTri, 0);
+			if (TID == FDynamicMesh3::NonManifoldID)
+			{
+				bSrcIsManifold = false;
+
+				int E0 = Mesh.FindEdge(AddTri[0], AddTri[1]);
+				int E1 = Mesh.FindEdge(AddTri[1], AddTri[2]);
+				int E2 = Mesh.FindEdge(AddTri[2], AddTri[0]);
+
+				// determine which verts need to be duplicated
+				bool bDuplicate[3] = { false, false, false };
+				if (E0 != FDynamicMesh3::InvalidID && Mesh.IsBoundaryEdge(E0) == false)
+				{
+					bDuplicate[0] = true;
+					bDuplicate[1] = true;
+				}
+				if (E1 != FDynamicMesh3::InvalidID && Mesh.IsBoundaryEdge(E1) == false)
+				{
+					bDuplicate[1] = true;
+					bDuplicate[2] = true;
+				}
+				if (E2 != FDynamicMesh3::InvalidID && Mesh.IsBoundaryEdge(E2) == false)
+				{
+					bDuplicate[2] = true;
+					bDuplicate[0] = true;
+				}
+				for (int32 SubIdx = 0; SubIdx < 3; ++SubIdx)
+				{
+					if (bDuplicate[SubIdx])
+					{
+						// duplicate vertex and all overlay elements
+						int32 ToDupVID = AddTri[SubIdx];
+						FVector3d Position = Mesh.GetVertex(ToDupVID);
+						int32 NewVID = Mesh.AppendVertex(Position);
+						int32 ColorEID = Mesh.Attributes()->PrimaryColors()->AppendElement(Mesh.Attributes()->PrimaryColors()->GetElement(ToDupVID));
+						checkSlow(ColorEID == NewVID);
+						for (int32 NormalLayerIdx = 0; NormalLayerIdx < 3; ++NormalLayerIdx)
+						{
+							const int32 NormalEID = Mesh.Attributes()->GetNormalLayer(NormalLayerIdx)->AppendElement(Mesh.Attributes()->GetNormalLayer(NormalLayerIdx)->GetElement(ToDupVID));
+							checkSlow(NormalEID == NewVID);
+						}
+						for (int32 UVLayer = 0; UVLayer < UVLayers.Num(); ++UVLayer)
+						{
+							const int32 UVElID = Mesh.Attributes()->GetUVLayer(UVLayer)->AppendElement(Mesh.Attributes()->GetUVLayer(UVLayer)->GetElement(ToDupVID));
+							checkSlow(UVElID == NewVID);
+						}
+						// set new vertex ID on triangle
+						AddTri[SubIdx] = NewVID;
+					}
+				}
+
+				// try to append again
+				TID = Mesh.AppendTriangle(AddTri);
+				checkSlow(TID != FDynamicMesh3::NonManifoldID);
+			}
+			if (TID < 0)
+			{
+				continue;
+			}
+			// after successfully appending the triangle, set overlay elements to match (since overlay elements are 1:1 with vertices in initial mesh)
+			Mesh.Attributes()->PrimaryColors()->SetTriangle(TID, AddTri, false);
+			for (int32 NormalLayerIdx = 0; NormalLayerIdx < 3; ++NormalLayerIdx)
+			{
+				Mesh.Attributes()->GetNormalLayer(NormalLayerIdx)->SetTriangle(TID, AddTri, false);
+			}
+			for (int32 UVLayerIdx = 0; UVLayerIdx < UVLayers.Num(); ++UVLayerIdx)
+			{
+				Mesh.Attributes()->GetUVLayer(UVLayerIdx)->SetTriangle(TID, AddTri, false);
+			}
+
+			Mesh.Attributes()->GetMaterialID()->SetValue(TID, MaterialIDs[FaceIdx]);
+			if (InternalFacePolygroups)
+			{
+				InternalFacePolygroups->SetValue(TID, 1 + (int32)InternalArray[FaceIdx]);
+			}
+			if (VisibleFacePolygroups)
+			{
+				VisibleFacePolygroups->SetValue(TID, 1 + (int32)VisibleArray[FaceIdx]);
+			}
+			// note: material index doesn't need to be passed through; will be rebuilt by a call to reindex materials once the cut mesh is returned back to geometry collection format
+		}
+
+		if (!Options.bSaveIsolatedVertices)
+		{
+			FDynamicMeshEditor Editor(&Mesh);
+			Editor.RemoveIsolatedVertices();
+		}
+
+		if (Options.bWeldVertices)
+		{
+			FMergeCoincidentMeshEdges Welder(&Mesh);
+			Welder.MergeVertexTolerance = FMathd::Epsilon;
+			Welder.bWeldAttrsOnMergedEdges = true;
+			Welder.Apply();
+		}
+
+		// if isolated verts were removed or edges were welded, may benefit from compacting the mesh
+		Mesh.CompactInPlace();
+	}
+	
+	return true;
+}
+
+// Convert all geometry in the given Collection to dynamic meshes
+bool FGeometryCollectionToDynamicMeshes::Init(const FManagedArrayCollection& Collection, const FToMeshOptions& Options)
+{
+	return InitHelper(Collection, true, GetCollectionTransforms(Collection), true, TConstArrayView<int32>(), Options);
+}
+// Convert geometry at selected transform indices in the given Collection to dynamic meshes
+bool FGeometryCollectionToDynamicMeshes::InitFromTransformSelection(const FManagedArrayCollection& Collection, TConstArrayView<int32> TransformIndices, const FToMeshOptions& Options)
+{
+	return InitHelper(Collection, true, GetCollectionTransforms(Collection), true, TransformIndices, Options);
+}
+
+namespace
+{
+	// Track instances of parent vertices with unique element IDs
+	template<int MaxUVLayers = 8>
+	struct TUniqueMeshVertex
+	{
+		// Note: We don't include the parent VID here, so verts with fully unset attribs will not be distinguishable (e.g., elements w/ no triangle)
+
+		// Element IDs for UV overlays
+		TStaticArray<int32, MaxUVLayers> UVsEIDs;
+		// Element IDs for Normal, Tangent and BiTangent overlays
+		TStaticArray<int32, 3> NormalEIDs;
+		// Element ID for color overlay
+		int32 ColorEID;
+
+		TUniqueMeshVertex()
+		{
+			ColorEID = -1;
+			NormalEIDs = { -1,-1,-1 };
+			for (int32 UVIdx = 0; UVIdx < MaxUVLayers; ++UVIdx)
+			{
+				UVsEIDs[UVIdx] = -1;
+			}
+		}
+		TUniqueMeshVertex(const FDynamicMesh3& Mesh, int32 TID, int32 SubIdx)
+		{
+			NormalEIDs[0] = Mesh.Attributes()->PrimaryNormals()->GetTriangle(TID)[SubIdx];
+			NormalEIDs[1] = Mesh.Attributes()->PrimaryTangents()->GetTriangle(TID)[SubIdx];
+			NormalEIDs[2] = Mesh.Attributes()->PrimaryBiTangents()->GetTriangle(TID)[SubIdx];
+			ColorEID = Mesh.Attributes()->PrimaryColors()->GetTriangle(TID)[SubIdx];
+			int32 NumUVs = FMath::Min(MaxUVLayers, Mesh.Attributes()->NumUVLayers()); // max 8 UV layers supported
+			for (int32 UVIdx = 0; UVIdx < NumUVs; ++UVIdx)
+			{
+				UVsEIDs[UVIdx] = (Mesh.Attributes()->GetUVLayer(UVIdx)->GetTriangle(TID)[SubIdx]);
+			}
+			for (int32 UVIdx = NumUVs; UVIdx < MaxUVLayers; ++UVIdx)
+			{
+				UVsEIDs[UVIdx] = INDEX_NONE;
+			}
+		}
+
+		bool operator==(const TUniqueMeshVertex& Other) const
+		{
+
+			bool bEquality = true;
+			for (int32 UVLayerIdx = 0; UVLayerIdx < MaxUVLayers; ++UVLayerIdx)
+			{
+				bEquality &= (this->UVsEIDs[UVLayerIdx] == Other.UVsEIDs[UVLayerIdx]);
+			}
+			bEquality &= (this->NormalEIDs[0] == Other.NormalEIDs[0]);
+			bEquality &= (this->NormalEIDs[1] == Other.NormalEIDs[1]);
+			bEquality &= (this->NormalEIDs[2] == Other.NormalEIDs[2]);
+			bEquality &= (this->ColorEID == Other.ColorEID);
+
+
+			return bEquality;
+		}
+	};
+
+}
+// Convert currently-held meshes back to geometry in the Collection, overwriting existing geometry as needed
+bool FGeometryCollectionToDynamicMeshes::UpdateGeometryCollection(FGeometryCollection& Collection, const FToCollectionOptions& Options) const
+{
+	bool bAllSucceeded = true;
+
+	int32 NumGeometry = Collection.NumElements(FGeometryCollection::GeometryGroup);
+	TArray<int32> NewFaceCounts, NewVertexCounts;
+	NewFaceCounts.SetNumUninitialized(NumGeometry);
+	NewVertexCounts.SetNumUninitialized(NumGeometry);
+	for (int32 GeomIdx = 0; GeomIdx < Collection.FaceCount.Num(); GeomIdx++)
+	{
+		NewFaceCounts[GeomIdx] = Collection.FaceCount[GeomIdx];
+		NewVertexCounts[GeomIdx] = Collection.VertexCount[GeomIdx];
+	}
+	bool bNeedsSplits = false, bNeedsResize = false;
+	
+	TArray<TUniqueMeshVertex<>> Uniques;
+	auto CountUniqueVerts = [&Uniques](const FDynamicMesh3& Mesh)
+	{
+		int32 Count = 0;
+		for (int32 VID : Mesh.VertexIndicesItr())
+		{
+			Uniques.Reset();
+			for (int32 TID : Mesh.VtxTrianglesItr(VID))
+			{
+				FIndex3i Tri = Mesh.GetTriangle(TID);
+				int32 SubIdx = Tri.IndexOf(VID);
+				Uniques.AddUnique(TUniqueMeshVertex<>(Mesh, TID, SubIdx));
+			}
+			Count += FMath::Max(1, Uniques.Num());
+		}
+		return Count;
+	};
+	for (int MeshIdx = 0; MeshIdx < Meshes.Num(); MeshIdx++)
+	{
+		const FMeshInfo& MeshData = Meshes[MeshIdx];
+		if (MeshData.Mesh.IsValid() && Collection.TransformToGeometryIndex.IsValidIndex(MeshData.TransformIndex))
+		{
+			int32 GeomIdx = Collection.TransformToGeometryIndex[MeshData.TransformIndex];
+			int32 PrevTriCount = NewFaceCounts[GeomIdx];
+			NewFaceCounts[GeomIdx] = MeshData.Mesh->TriangleCount();
+			bNeedsResize = bNeedsResize || NewFaceCounts[GeomIdx] != PrevTriCount;
+			int32 PrevVertexCount = NewVertexCounts[GeomIdx];
+			NewVertexCounts[GeomIdx] = CountUniqueVerts(*MeshData.Mesh);
+			bNeedsSplits = bNeedsSplits || NewVertexCounts[GeomIdx] != MeshData.Mesh->VertexCount();
+			bNeedsResize = bNeedsResize || NewVertexCounts[GeomIdx] != PrevVertexCount;
+		}
+	}
+
+	if (bNeedsResize)
+	{
+		bool bDoValidation = false;
+#if UE_BUILD_DEBUG
+		bDoValidation = true; // note: this validation is extremely slow
+#endif
+		GeometryCollectionAlgo::ResizeGeometries(&Collection, NewFaceCounts, NewVertexCounts, bDoValidation);
+	}
+
+	for (int MeshIdx = 0; MeshIdx < Meshes.Num(); MeshIdx++)
+	{
+		const FMeshInfo& MeshData = Meshes[MeshIdx];
+		const FDynamicMesh3& Mesh = *MeshData.Mesh;
+		int32 GeometryIdx = Collection.TransformToGeometryIndex[MeshData.TransformIndex];
+		bool bSucceeded = UpdateCollection(MeshData.Transform, Mesh, GeometryIdx, Collection, Options);
+		bAllSucceeded &= bSucceeded;
+	}
+
+	return bAllSucceeded;
+}
+
+
+// Update an existing geometry in a collection w/ a new mesh (w/ the same number of faces and vertices!)
+bool FGeometryCollectionToDynamicMeshes::UpdateCollection(const FTransform& FromCollection, const FDynamicMesh3& Mesh, int32 GeometryIdx, FGeometryCollection& Output, const FToCollectionOptions& Options)
+{
+	int32 VertexCount = Output.VertexCount[GeometryIdx];
+	int32 VertexStart = Output.VertexStart[GeometryIdx];
+	int32 TriangleCount = Output.FaceCount[GeometryIdx];
+	int32 TriangleStart = Output.FaceStart[GeometryIdx];
+
+	int32 UVLayerCount = Mesh.HasAttributes() ? Mesh.Attributes()->NumUVLayers() : 1;
+	Output.SetNumUVLayers(UVLayerCount);
+	GeometryCollection::UV::FUVLayers OutputUVLayers = GeometryCollection::UV::FindActiveUVLayers(Output);
+
+	int32 TransformIdx = Output.TransformIndex[GeometryIdx];
+
+	FBox Bounds;
+
+	// map from mesh vertex IDs to geometry collection indices
+	TArray<int32> VertexIDToIdxStart;
+	VertexIDToIdxStart.Init(INDEX_NONE, Mesh.MaxVertexID());
+	TArray<TUniqueMeshVertex<>> PerIdxElements;
+	PerIdxElements.Reserve(VertexCount);
+	// Temp storage for unique elements on a given vertex
+	TArray<TUniqueMeshVertex<>> Uniques;
+	{
+		int32 Count = 0;
+		for (int32 VID : Mesh.VertexIndicesItr())
+		{
+			VertexIDToIdxStart[VID] = PerIdxElements.Num();
+			Uniques.Reset();
+			FVector3f Pos = (FVector3f)FromCollection.InverseTransformPosition(FVector(Mesh.GetVertex(VID)));
+			Bounds += (FVector)Pos;
+			for (int32 TID : Mesh.VtxTrianglesItr(VID))
+			{
+				FIndex3i Tri = Mesh.GetTriangle(TID);
+				int32 SubIdx = Tri.IndexOf(VID);
+				Uniques.AddUnique(TUniqueMeshVertex<>(Mesh, TID, SubIdx));
+			}
+			if (Uniques.IsEmpty())
+			{
+				Uniques.Emplace();
+			}
+			for (int32 SubIdx = 0; SubIdx < Uniques.Num(); ++SubIdx)
+			{
+				int32 CopyToIdx = VertexStart + PerIdxElements.Add(Uniques[SubIdx]);
+				Output.Vertex[CopyToIdx] = Pos;
+				Output.BoneMap[CopyToIdx] = TransformIdx;
+				FVector3f NormalVals[3]{ FVector3f::ZAxisVector, FVector3f::XAxisVector, FVector3f::YAxisVector };
+				if (Mesh.HasAttributes())
+				{
+					int32 NumNormalLayers = FMath::Min(3, Mesh.Attributes()->NumNormalLayers());
+					for (int32 NormalLayerIdx = 0; NormalLayerIdx < NumNormalLayers; ++NormalLayerIdx)
+					{
+						int32 ElID = Uniques[SubIdx].NormalEIDs[NormalLayerIdx];
+						if (ElID > -1)
+						{
+							NormalVals[NormalLayerIdx] = Mesh.Attributes()->GetNormalLayer(NormalLayerIdx)->GetElement(ElID);
+						}
+					}
+				}
+				Output.Normal[CopyToIdx] = (FVector3f)FromCollection.InverseTransformVectorNoScale(FVector(NormalVals[0]));
+				Output.TangentU[CopyToIdx] = (FVector3f)FromCollection.InverseTransformVectorNoScale(FVector(NormalVals[1]));
+				Output.TangentV[CopyToIdx] = (FVector3f)FromCollection.InverseTransformVectorNoScale(FVector(NormalVals[2]));
+				if (const FDynamicMeshColorOverlay* Colors = Mesh.Attributes()->PrimaryColors())
+				{
+					int32 ElID = Uniques[SubIdx].ColorEID;
+					if (ElID > -1)
+					{
+						Output.Color[CopyToIdx] = FLinearColor(Colors->GetElement(ElID));
+					}
+				}
+
+				for (int32 UVLayer = 0; UVLayer < UVLayerCount; ++UVLayer)
+				{
+					FVector2f UV(0, 0);
+					if (const FDynamicMeshUVOverlay* UVs = Mesh.Attributes()->GetUVLayer(UVLayer))
+					{
+						int32 ElID = Uniques[SubIdx].UVsEIDs[UVLayer];
+						if (ElID > -1)
+						{
+							UV = UVs->GetElement(ElID);
+						}
+					}
+					OutputUVLayers[UVLayer][CopyToIdx] = UV;
+				}
+			}
+		}
+
+
+	}
+
+	check(PerIdxElements.Num() == VertexCount);
+
+	const FDynamicMeshPolygroupAttribute* VisLayer = nullptr;
+	const FDynamicMeshPolygroupAttribute* InternalLayer = nullptr;
+
+	if (Mesh.HasAttributes())
+	{
+		for (int32 Idx = 0; Idx < Mesh.Attributes()->NumPolygroupLayers(); ++Idx)
+		{
+			const FDynamicMeshPolygroupAttribute* Layer = Mesh.Attributes()->GetPolygroupLayer(Idx);
+			if (Layer->GetName() == VisibleFacePolyGroupName())
+			{
+				VisLayer = Layer;
+			}
+			if (Layer->GetName() == InternalFacePolyGroupName())
+			{
+				InternalLayer = Layer;
+			}
+		}
+	}
+
+	FIntVector VertexStartOffset(VertexStart);
+	int32 TriIdxOffset = 0;
+	for (int32 TID : Mesh.TriangleIndicesItr())
+	{
+		int32 CopyToIdx = TriangleStart + TriIdxOffset;
+		TriIdxOffset++;
+		bool bVis = Options.bDefaultFaceVisible, bInternal = Options.bDefaultFaceInternal;
+		int32 MatID = 0;
+		if (Mesh.HasAttributes())
+		{
+			if (VisLayer)
+			{
+				int32 VisTag = VisLayer->GetValue(TID) - 1;
+				if (VisTag > -1)
+				{
+					bVis = bool(VisTag);
+				}
+			}
+			if (InternalLayer)
+			{
+				int32 InternalTag = InternalLayer->GetValue(TID) - 1;
+				if (InternalTag > -1)
+				{
+					bInternal = bool(InternalTag);
+				}	
+			}
+			if (const FDynamicMeshMaterialAttribute* Mats = Mesh.Attributes()->GetMaterialID())
+			{
+				MatID = Mats->GetValue(TID);
+			}
+		}
+
+		Output.Visible[CopyToIdx] = bVis;
+		Output.Internal[CopyToIdx] = bInternal;
+		Output.MaterialID[CopyToIdx] = MatID;
+		FIndex3i Tri = Mesh.GetTriangle(TID);
+		FIntVector OutTri;
+		for (int32 SubIdx = 0; SubIdx < 3; ++SubIdx)
+		{
+			TUniqueMeshVertex<> Unique(Mesh, TID, SubIdx);
+			int32 MeshVID = Tri[SubIdx];
+			int32 StartIdx = VertexIDToIdxStart[MeshVID];
+			int32 EndIdx = MeshVID + 1 < VertexIDToIdxStart.Num() ? VertexIDToIdxStart[MeshVID + 1] : PerIdxElements.Num();
+			int32 FoundIdx = -1;
+			for (int32 TestIdx = StartIdx; TestIdx < EndIdx; ++TestIdx)
+			{
+				if (PerIdxElements[TestIdx] == Unique)
+				{
+					FoundIdx = TestIdx;
+					break;
+				}
+			}
+			if (ensure(FoundIdx > -1))
+			{
+				OutTri[SubIdx] = FoundIdx + VertexStart;
+			}
+			else // (should not happen)
+			{
+				OutTri[SubIdx] = StartIdx + VertexStart;
+			}
+		}
+		Output.Indices[CopyToIdx] = OutTri;
+	}
+	ensure(TriIdxOffset == TriangleCount);
+	if (Output.BoundingBox.Num())
+	{
+		Output.BoundingBox[GeometryIdx] = Bounds;
+	}
+
+	return true;
+}
+
+
+FName FGeometryCollectionToDynamicMeshes::InternalFacePolyGroupName()
+{
+	return FName("GeometryCollectionInternalFaces");
+}
+FName FGeometryCollectionToDynamicMeshes::VisibleFacePolyGroupName()
+{
+	return FName("GeometryCollectionVisibleFaces");
+}
+
+} // namespace UE::Geometry
\ No newline at end of file
diff --git a/Engine/Source/Runtime/MeshConversion/Public/GeometryCollectionToDynamicMesh.h b/Engine/Source/Runtime/MeshConversion/Public/GeometryCollectionToDynamicMesh.h
new file mode 100644
--- /dev/null
+++ b/Engine/Source/Runtime/MeshConversion/Public/GeometryCollectionToDynamicMesh.h
@@ -0,0 +1,81 @@
+// Copyright Epic Games, Inc. All Rights Reserved.
+
+#pragma once
+
+#include "CoreMinimal.h"
+
+#include "GeometryBase.h"
+
+PREDECLARE_GEOMETRY(class FDynamicMesh3);
+struct FManagedArrayCollection;
+class FGeometryCollection;
+
+namespace UE::Geometry
+{
+
+struct FGeometryCollectionToDynamicMeshes
+{
+	struct FMeshInfo
+	{
+		int32 TransformIndex = -1;
+		TUniquePtr<FDynamicMesh3> Mesh;
+		FTransform Transform = FTransform::Identity;
+	};
+
+	TArray<FMeshInfo> Meshes;
+
+
+	
+	enum class EInvisibleFaceConversion : uint8
+	{
+		Skip,
+		TagWithPolygroup
+	};
+
+	struct FToMeshOptions
+	{
+		FTransform Transform = FTransform::Identity;
+
+		bool bWeldVertices = false;
+		bool bSaveIsolatedVertices = false;
+
+		bool bInternalFaceTagsAsPolygroups = true;
+		EInvisibleFaceConversion InvisibleFaces = EInvisibleFaceConversion::Skip;
+	};
+
+	struct FToCollectionOptions
+	{
+		// Whether to set faces with no 'visible' tag as visible (if true) or invisible (if false)
+		bool bDefaultFaceVisible = true;
+		// Whether to set faces with no 'internal' tag as internal (if true) or external (if false)
+		bool bDefaultFaceInternal = false;
+	};
+
+	// Convert all meshes in the given collection to dynamic meshes
+	MESHCONVERSION_API bool Init(const FManagedArrayCollection& Collection, const FToMeshOptions& Options);
+
+	// Convert the meshes in the collection at the given transforms to dynamic meshes
+	MESHCONVERSION_API bool InitFromTransformSelection(const FManagedArrayCollection& Collection, TConstArrayView<int32> TransformIndices, const FToMeshOptions& Options);
+
+	// Update a geometry collection with the current meshes array
+	// Note this updates a Geometry Collection rather than the more general FManagedArrayCollection because some of the updating code is specific to geometry collection
+	MESHCONVERSION_API bool UpdateGeometryCollection(FGeometryCollection& Collection, const FToCollectionOptions& Options) const;
+
+
+
+	// Get the name used for the polygroup that we optionally set on the dynamic mesh, corresponding to the internal face tags of the geometry collection triangles
+	static MESHCONVERSION_API FName InternalFacePolyGroupName();
+	// Get the name used for the polygroup that we optionally set on the dynamic mesh, corresponding to the invisible face tags of the geometry collection triangles
+	static MESHCONVERSION_API FName VisibleFacePolyGroupName();
+
+private:
+	
+	// private helper method with shared implementation for all Init* methods
+	bool InitHelper(const FManagedArrayCollection& Collection, bool bTransformInComponentSpace, TConstArrayView<FTransform3f> Transforms, bool bAllTransforms, TConstArrayView<int32> TransformIndices, const FToMeshOptions& Options);
+	static TConstArrayView<FTransform3f> GetCollectionTransforms(const FManagedArrayCollection& Collection);
+
+	// Update an existing geometry in a collection w/ a new mesh (w/ the same number of faces and vertices!)
+	static bool UpdateCollection(const FTransform& FromCollection, const FDynamicMesh3& Mesh, int32 GeometryIdx, FGeometryCollection& Output, const FToCollectionOptions& Options);
+};
+
+} // namespace UE::Geometry::GeometryCollectionDynamicMeshConversion
\ No newline at end of file
