clear the output dynamic mesh in meshdescription->dynamicmesh conversion, since the code assumes it is empty

diff --git a/Engine/Source/Runtime/MeshConversion/Private/MeshDescriptionToDynamicMesh.cpp b/Engine/Source/Runtime/MeshConversion/Private/MeshDescriptionToDynamicMesh.cpp
--- a/Engine/Source/Runtime/MeshConversion/Private/MeshDescriptionToDynamicMesh.cpp
+++ b/Engine/Source/Runtime/MeshConversion/Private/MeshDescriptionToDynamicMesh.cpp
@@ -178,6 +178,7 @@ void FMeshDescriptionToDynamicMesh::Convert(const FMeshDescription* MeshIn, FDyn
 
 	TriIDMap.Reset();
 	VertIDMap.Reset();
+	MeshOut.Clear();
 
 	if (!ensure(MeshIn != nullptr))
 	{
