Fix wrong function call in MeshDescriptionUVsToDynamicMesh convert logic

diff --git a/Engine/Source/Runtime/MeshConversion/Private/MeshDescriptionUVsToDynamicMesh.cpp b/Engine/Source/Runtime/MeshConversion/Private/MeshDescriptionUVsToDynamicMesh.cpp
--- a/Engine/Source/Runtime/MeshConversion/Private/MeshDescriptionUVsToDynamicMesh.cpp
+++ b/Engine/Source/Runtime/MeshConversion/Private/MeshDescriptionUVsToDynamicMesh.cpp
@@ -126,7 +126,7 @@ TSharedPtr<FDynamicMesh3> FMeshDescriptionUVsToDynamicMesh::GetUVMesh(const FMes
 			const FVector2f UV = SharedUVCoordinates[UVID];
 			MeshOut->InsertVertex(UVID.GetValue(), UVToVertPosition(UV, ScaleFactor), true);
 		}
-		MeshOut->EndUnsafeTrianglesInsert();
+		MeshOut->EndUnsafeVerticesInsert();
 	}
 
 	// Instance UV's and a welder are used if bUseSharedUVs is false, or if it is true but some 
