stop creating unused temp array of edges when adding triangles to mesh description via MeshDescriptionBuilder

diff --git a/Engine/Source/Runtime/MeshConversion/Private/MeshDescriptionBuilder.cpp b/Engine/Source/Runtime/MeshConversion/Private/MeshDescriptionBuilder.cpp
--- a/Engine/Source/Runtime/MeshConversion/Private/MeshDescriptionBuilder.cpp
+++ b/Engine/Source/Runtime/MeshConversion/Private/MeshDescriptionBuilder.cpp
@@ -271,8 +271,7 @@ FTriangleID FMeshDescriptionBuilder::AppendTriangle(const FVertexInstanceID& Ins
 	CornerInstanceIDs.Add(Instance1);
 	CornerInstanceIDs.Add(Instance2);
 
-	TArray<FEdgeID> NewEdgeIDs; 
-	const FTriangleID NewTriangleID = MeshDescription->CreateTriangle(PolygonGroup, CornerInstanceIDs, &NewEdgeIDs);
+	const FTriangleID NewTriangleID = MeshDescription->CreateTriangle(PolygonGroup, CornerInstanceIDs);
 
 	return NewTriangleID;
 }
