Fix static mesh's LOD -> dynamic mesh conversion handling of non-compact polygon group IDs -- we need to compact these IDs so that they correspond to static mesh section indices, which all our tools and geometry script processing expect.

To reduce risk of this affecting unrelated processing, this compacting is done via an optional flag, which is enabled for the paths where we already have special code to handle static mesh materials (i.e. places only where SetPolygonGroupToMaterialIndexMap could be called)

diff --git a/Engine/Plugins/Runtime/MeshModelingToolset/Source/ModelingComponents/Private/ConversionUtils/DynamicMeshViaMeshDescriptionUtil.cpp b/Engine/Plugins/Runtime/MeshModelingToolset/Source/ModelingComponents/Private/ConversionUtils/DynamicMeshViaMeshDescriptionUtil.cpp
--- a/Engine/Plugins/Runtime/MeshModelingToolset/Source/ModelingComponents/Private/ConversionUtils/DynamicMeshViaMeshDescriptionUtil.cpp
+++ b/Engine/Plugins/Runtime/MeshModelingToolset/Source/ModelingComponents/Private/ConversionUtils/DynamicMeshViaMeshDescriptionUtil.cpp
@@ -23,6 +23,7 @@ FDynamicMesh3 UE::Geometry::GetDynamicMeshViaMeshDescription(
 	FDynamicMesh3 DynamicMesh;
 	FMeshDescriptionToDynamicMesh Converter;
 	Converter.bVIDsFromNonManifoldMeshDescriptionAttr = true;
+	Converter.bUseCompactedPolygonGroupIDValues = true;
 	Converter.SetPolygonGroupToMaterialIndexMap(MeshDescriptionProvider.GetPolygonGroupToMaterialIndexMap());
 	if (InGetMeshParams.bWantMeshTangents)
 	{
diff --git a/Engine/Plugins/Runtime/MeshModelingToolset/Source/ModelingComponents/Private/ConversionUtils/SceneComponentToDynamicMesh.cpp b/Engine/Plugins/Runtime/MeshModelingToolset/Source/ModelingComponents/Private/ConversionUtils/SceneComponentToDynamicMesh.cpp
--- a/Engine/Plugins/Runtime/MeshModelingToolset/Source/ModelingComponents/Private/ConversionUtils/SceneComponentToDynamicMesh.cpp
+++ b/Engine/Plugins/Runtime/MeshModelingToolset/Source/ModelingComponents/Private/ConversionUtils/SceneComponentToDynamicMesh.cpp
@@ -275,6 +275,7 @@ namespace Private::ConversionHelper
 
 		FMeshDescriptionToDynamicMesh Converter;
 		Converter.bVIDsFromNonManifoldMeshDescriptionAttr = AssetOptions.bIncludeNonManifoldSrcInfo;
+		Converter.bUseCompactedPolygonGroupIDValues = true;
 		if (!AssetOptions.bUseSectionMaterialIndices)
 		{
 			Converter.SetPolygonGroupToMaterialIndexMap(PolygonGroupToMaterialMap);
diff --git a/Engine/Plugins/Runtime/MeshModelingToolset/Source/ModelingComponents/Private/ModelingToolTargetUtil.cpp b/Engine/Plugins/Runtime/MeshModelingToolset/Source/ModelingComponents/Private/ModelingToolTargetUtil.cpp
--- a/Engine/Plugins/Runtime/MeshModelingToolset/Source/ModelingComponents/Private/ModelingToolTargetUtil.cpp
+++ b/Engine/Plugins/Runtime/MeshModelingToolset/Source/ModelingComponents/Private/ModelingToolTargetUtil.cpp
@@ -286,6 +286,7 @@ FDynamicMesh3 UE::ToolTarget::GetDynamicMeshCopy(UToolTarget* Target, const FGet
 	{
 		FMeshDescriptionToDynamicMesh Converter;
 		Converter.bVIDsFromNonManifoldMeshDescriptionAttr = true;
+		Converter.bUseCompactedPolygonGroupIDValues = true;
 		Converter.SetPolygonGroupToMaterialIndexMap(MeshDescriptionProvider->GetPolygonGroupToMaterialIndexMap());
 		if (InGetMeshParams.bWantMeshTangents)
 		{
diff --git a/Engine/Source/Runtime/MeshConversion/Private/MeshDescriptionToDynamicMesh.cpp b/Engine/Source/Runtime/MeshConversion/Private/MeshDescriptionToDynamicMesh.cpp
--- a/Engine/Source/Runtime/MeshConversion/Private/MeshDescriptionToDynamicMesh.cpp
+++ b/Engine/Source/Runtime/MeshConversion/Private/MeshDescriptionToDynamicMesh.cpp
@@ -217,6 +217,24 @@ void FMeshDescriptionToDynamicMesh::Convert(const FMeshDescription* MeshIn, FDyn
 		MeshOut.EnableTriangleGroups(0);
 	}
 
+	// If needed, compute mapping to 'compacted' versions of the Polygon Group IDs
+	TArray<int32> CompactedPolygonGroupMap;
+	if (bUseCompactedPolygonGroupIDValues && MeshIn->PolygonGroups().Num() != MeshIn->PolygonGroups().GetArraySize())
+	{
+		const FPolygonGroupArray& PolygonGroups = MeshIn->PolygonGroups();
+		CompactedPolygonGroupMap.Init(INDEX_NONE, PolygonGroups.GetArraySize());
+		int32 SectionIndex = 0;
+		for (FPolygonGroupID GroupID : PolygonGroups.GetElementIDs())
+		{
+			CompactedPolygonGroupMap[GroupID.GetValue()] = SectionIndex;
+			++SectionIndex;
+		}
+	}
+	auto MapPolygonGroupID = [&CompactedPolygonGroupMap](FPolygonGroupID GroupID) -> int32
+	{
+		return CompactedPolygonGroupMap.IsEmpty() ? GroupID.GetValue() : CompactedPolygonGroupMap[GroupID.GetValue()];
+	};
+
 	bCopyTangents = bCopyTangents && InstanceTangents.IsValid() && InstanceBiTangentSign.IsValid();
 	
 	// base triangle groups will track polygons
@@ -255,7 +273,7 @@ void FMeshDescriptionToDynamicMesh::Convert(const FMeshDescription* MeshIn, FDyn
 		// Get the PolygonID, PolygonGroupID and general GroupID
 		//---
 		FPolygonID PolygonID = MeshIn->GetTrianglePolygon(TriangleID);
-		FPolygonGroupID PolygonGroupID = MeshIn->GetTrianglePolygonGroup(TriangleID);
+		int32 MappedPolygonGroupID = MapPolygonGroupID(MeshIn->GetTrianglePolygonGroup(TriangleID));
 
 		int GroupID = 0;
 		switch (GroupMode)
@@ -270,7 +288,7 @@ void FMeshDescriptionToDynamicMesh::Convert(const FMeshDescription* MeshIn, FDyn
 				GroupID = PolygonID.GetValue() + 1; // Shift IDs up by 1 to leave ID 0 as a default/unassigned group
 				break;
 			case EPrimaryGroupMode::SetToPolygonGroupID:
-				GroupID = PolygonGroupID.GetValue() + 1; // Shift IDs up by 1 to leave ID 0 as a default/unassigned group
+				GroupID = MappedPolygonGroupID + 1; // Shift IDs up by 1 to leave ID 0 as a default/unassigned group
 				break;
 			case EPrimaryGroupMode::SetToZero:
 				break; // keep at 0
@@ -278,7 +296,7 @@ void FMeshDescriptionToDynamicMesh::Convert(const FMeshDescription* MeshIn, FDyn
 		
 		FTriData TriData;
 		TriData.PolygonID = PolygonID;
-		TriData.PolygonGroupID = PolygonGroupID;
+		TriData.PolygonGroupID = MappedPolygonGroupID;
 		
 
 		// stash the vertex instance IDs for this triangle.  potentially needed for per-instance attribute welding
diff --git a/Engine/Source/Runtime/MeshConversion/Public/MeshDescriptionToDynamicMesh.h b/Engine/Source/Runtime/MeshConversion/Public/MeshDescriptionToDynamicMesh.h
--- a/Engine/Source/Runtime/MeshConversion/Public/MeshDescriptionToDynamicMesh.h
+++ b/Engine/Source/Runtime/MeshConversion/Public/MeshDescriptionToDynamicMesh.h
@@ -42,7 +42,21 @@ public:
 	/** Should Vertex Colors of MeshDescription be transformed from Linear to SRGB */
 	bool bTransformVertexColorsLinearToSRGB = true;
 
-	/** Set an optional mapping from mesh description polygon group to dynamic mesh material indices; if unset, an identity map is assumed */
+	/** 
+	 * Set this flag to transform Polygon Group IDs to their compacted indices, before using them as material indices --
+	 * in other words, 'holes' in the Mesh Description's PolygonGroups array are removed by shifting the subsequent IDs down
+	 * 
+	 * Note: For Static Mesh LODs, this transforms Polygon Group IDs to Section Indices.
+	 * This flag should typically be set to 'true' when converting Static Mesh assets.
+	 */
+	bool bUseCompactedPolygonGroupIDValues = false;
+
+	/**
+	 * Set an optional mapping from mesh description polygon group to dynamic mesh material indices; if unset, an identity map is assumed 
+	 *
+	 * Note: If bUseCompactedPolygonGroupIDValues is true, this remapping applies *after* the compact step --
+	 *  so this mapping is a Section Index -> Material Index map for Static Mesh assets.
+	 */
 	void SetPolygonGroupToMaterialIndexMap(const TArray<int32>& PolygonGroupToMaterialIndexMapIn)
 	{
 		PolygonGroupToMaterialIndexMap = PolygonGroupToMaterialIndexMapIn;
