Fix for hard edges / normal seams being lost on recompute normals after dynamic mesh -> mesh description conversion:
Add option for dynamic mesh -> mesh description conversion to tag normal seams as hard edges, so they are similarly split if the static or skeletal mesh side recomputes normals

diff --git a/Engine/Source/Runtime/MeshConversion/Private/DynamicMeshToMeshDescription.cpp b/Engine/Source/Runtime/MeshConversion/Private/DynamicMeshToMeshDescription.cpp
--- a/Engine/Source/Runtime/MeshConversion/Private/DynamicMeshToMeshDescription.cpp
+++ b/Engine/Source/Runtime/MeshConversion/Private/DynamicMeshToMeshDescription.cpp
@@ -11,6 +11,7 @@
 #include "DynamicMesh/MeshTangents.h"
 #include "DynamicMesh/NonManifoldMappingSupport.h"
 #include "Util/ColorConstants.h"
+#include "Async/ParallelFor.h"
 
 using namespace UE::Geometry;
 
@@ -1109,6 +1110,47 @@ void FDynamicMeshToMeshDescription::Convert_NoSharedInstances(const FDynamicMesh
 	// convert weight map layers
 	ConvertWeightLayers(MeshIn, MeshOut, MapV);
 
+	// convert normal seams to hard edges
+	if (ConversionOptions.bSetHardEdgesFromNormalSeams && NormalOverlay)
+	{
+		FSkeletalMeshAttributes MeshOutAttributes(MeshOut);
+		TEdgeAttributesRef<bool> EdgeHardness = MeshOutAttributes.GetEdgeHardnesses();
+
+		// Find seam edges on the dynamic mesh
+		TArray<bool> MeshSeamEdges;
+		MeshSeamEdges.SetNumZeroed(MeshIn->MaxEdgeID());
+		ParallelFor(MeshIn->MaxEdgeID(), [&MeshIn, &NormalOverlay, &MeshSeamEdges](int32 EID)
+			{
+				if (MeshIn->IsEdge(EID))
+				{
+					MeshSeamEdges[EID] = NormalOverlay->IsSeamEdge(EID);
+				}
+			});
+
+		// For any seam edge, find the corresponding mesh description edge and tag it as hard
+		// Note that this will also tag non-manifold edges (if preserved via remapping) as hard edges
+		for (int32 EID = 0; EID < MeshSeamEdges.Num(); ++EID)
+		{
+			if (MeshSeamEdges[EID]) // will only be true for valid edges that are boundaries or seams
+			{
+				int32 MeshInTID = MeshIn->GetEdgeT(EID).A;
+				int32 EdgeTriSubIdx = MeshIn->GetTriEdges(MeshInTID).IndexOf(EID);
+				if (!ensure(EdgeTriSubIdx != IndexConstants::InvalidID))
+				{
+					continue;
+				}
+
+				FTriangleID MeshOutTID = IndexToTriangleIDMap[MeshInTID];
+				TArrayView<FEdgeID> OutTriEdges = MeshOut.GetTriangleEdges(MeshOutTID);
+				if (!ensure(OutTriEdges.IsValidIndex(EdgeTriSubIdx)))
+				{
+					continue;
+				}
+				EdgeHardness[OutTriEdges[EdgeTriSubIdx]] = true;
+			}
+		}
+	}
+
 	// Convert all attached skin weights, if we're converting a mesh description that originated from
 	// a USkeletalMesh.
 	if (bHasSkinWeightsAttributes)
diff --git a/Engine/Source/Runtime/MeshConversion/Public/MeshConversionOptions.h b/Engine/Source/Runtime/MeshConversion/Public/MeshConversionOptions.h
--- a/Engine/Source/Runtime/MeshConversion/Public/MeshConversionOptions.h
+++ b/Engine/Source/Runtime/MeshConversion/Public/MeshConversionOptions.h
@@ -33,6 +33,9 @@ public:
 	/** Should the mesh be put back in the original non-manifold state, if the information is present */
 	bool bConvertBackToNonManifold = false;
 
+	/** Should normal seams be tagged as hard edges in the mesh description */
+	bool bSetHardEdgesFromNormalSeams = true;
+
 	//
 	// utility functions for common configuration cases
 	//
