Make sure edge hardness attributes are registered if they might be set by the dynamic mesh -> mesh description conversion

diff --git a/Engine/Source/Runtime/MeshConversion/Private/DynamicMeshToMeshDescription.cpp b/Engine/Source/Runtime/MeshConversion/Private/DynamicMeshToMeshDescription.cpp
--- a/Engine/Source/Runtime/MeshConversion/Private/DynamicMeshToMeshDescription.cpp
+++ b/Engine/Source/Runtime/MeshConversion/Private/DynamicMeshToMeshDescription.cpp
@@ -799,6 +799,13 @@ void FDynamicMeshToMeshDescription::Convert_NoSharedInstances(const FDynamicMesh
 	FMeshDescriptionBuilder Builder;
 	Builder.SetMeshDescription(&MeshOut);
 
+	// Make sure edge hardnesses are registered, if they might be set below
+	if (ConversionOptions.bSetHardEdgesFromNormalSeams && NormalOverlay)
+	{
+		FStaticMeshAttributes MeshOutAttributes(MeshOut);
+		MeshOutAttributes.Register();
+	}
+
 	bool bCopyGroupToPolyGroup = false;
 	if (ConversionOptions.bSetPolyGroups && MeshIn->HasTriangleGroups())
 	{
@@ -1113,7 +1120,7 @@ void FDynamicMeshToMeshDescription::Convert_NoSharedInstances(const FDynamicMesh
 	// convert normal seams to hard edges
 	if (ConversionOptions.bSetHardEdgesFromNormalSeams && NormalOverlay)
 	{
-		FSkeletalMeshAttributes MeshOutAttributes(MeshOut);
+		FStaticMeshAttributes MeshOutAttributes(MeshOut);
 		TEdgeAttributesRef<bool> EdgeHardness = MeshOutAttributes.GetEdgeHardnesses();
 
 		// Find seam edges on the dynamic mesh
