add defensive ensure()s to prevent crash / log issues on dynamic mesh overlay -> static mesh attributes update

diff --git a/Engine/Source/Runtime/MeshConversion/Private/DynamicMeshToMeshDescription.cpp b/Engine/Source/Runtime/MeshConversion/Private/DynamicMeshToMeshDescription.cpp
--- a/Engine/Source/Runtime/MeshConversion/Private/DynamicMeshToMeshDescription.cpp
+++ b/Engine/Source/Runtime/MeshConversion/Private/DynamicMeshToMeshDescription.cpp
@@ -24,24 +24,50 @@ namespace DynamicMeshToMeshDescriptionConversionHelper
 		const FDynamicMesh3* MeshInArg, const FMeshDescription& MeshOutArg,
 		TVertexInstanceAttributesRef<OutAttributeType>& InstanceAttrib, const TDynamicMeshVectorOverlay<float, VecLen, InAttributeType>* Overlay, OutAttributeType DefaultValue, int AttribIndex=0)
 	{
+		if (!ensureMsgf(AttribIndex >= 0 && AttribIndex < InstanceAttrib.GetNumChannels(), TEXT("Failed to set attributes from overlay; invalid AttribIndex %d"), AttribIndex))
+		{
+			return;
+		}
 		for (const FTriangleID TriangleID : MeshOutArg.Triangles().GetElementIDs())
 		{
 			TArrayView<const FVertexInstanceID> InstanceTri = MeshOutArg.GetTriangleVertexInstances(TriangleID);
+			auto ValidTargetElement = [NumAttribElements = InstanceAttrib.GetNumElements()](const FVertexInstanceID ID) -> bool 
+			{
+				int32 IDValue = ID.GetValue();
+				return ensureMsgf(IDValue >= 0 && IDValue < NumAttribElements, TEXT("Failed to set attributes from overlay; invalid target vertex instance %d in mesh with %d instances"), IDValue, NumAttribElements);
+			};
+			if (!ValidTargetElement(InstanceTri[0]) || !ValidTargetElement(InstanceTri[1]) || !ValidTargetElement(InstanceTri[2]))
+			{
+				continue;
+			}
+			
+			auto SetDefaults = [&InstanceAttrib, &InstanceTri, AttribIndex, &DefaultValue]()
+			{
+				InstanceAttrib.Set(InstanceTri[0], AttribIndex, DefaultValue);
+				InstanceAttrib.Set(InstanceTri[1], AttribIndex, DefaultValue);
+				InstanceAttrib.Set(InstanceTri[2], AttribIndex, DefaultValue);
+			};
 
 			int32 MeshInTriIdx = TriangleID.GetValue();
 			
-			if (Overlay->IsSetTriangle(MeshInTriIdx))
+			if (Overlay->IsSetTriangle(MeshInTriIdx) && 
+				ensureMsgf(MeshInArg->IsTriangle(MeshInTriIdx), TEXT("Failed to set attributes from overlay; invalid TriIdx %d not found in mesh with %d Tri IDs and %d Tri Count"), MeshInTriIdx, MeshInArg->MaxTriangleID(), MeshInArg->TriangleCount()))
 			{
-				FIndex3i OverlayVertIndices = Overlay->GetTriangle(MeshInTriIdx);
-				InstanceAttrib.Set(InstanceTri[0], AttribIndex, OutAttributeType(Overlay->GetElement(OverlayVertIndices.A)));
-				InstanceAttrib.Set(InstanceTri[1], AttribIndex, OutAttributeType(Overlay->GetElement(OverlayVertIndices.B)));
-				InstanceAttrib.Set(InstanceTri[2], AttribIndex, OutAttributeType(Overlay->GetElement(OverlayVertIndices.C)));
+				FIndex3i OverlayEls = Overlay->GetTriangle(MeshInTriIdx);
+				if (!ensureMsgf(Overlay->IsElement(OverlayEls.A) && Overlay->IsElement(OverlayEls.B) && Overlay->IsElement(OverlayEls.C), TEXT("Failed to set attributes from overlay; invalid overlay triangle w/ indices (%d, %d, %d) on overlay with %d element IDs"), OverlayEls.A, OverlayEls.B, OverlayEls.C, Overlay->MaxElementID()))
+				{
+					SetDefaults();
+				}
+				else
+				{
+					InstanceAttrib.Set(InstanceTri[0], AttribIndex, OutAttributeType(Overlay->GetElement(OverlayEls.A)));
+					InstanceAttrib.Set(InstanceTri[1], AttribIndex, OutAttributeType(Overlay->GetElement(OverlayEls.B)));
+					InstanceAttrib.Set(InstanceTri[2], AttribIndex, OutAttributeType(Overlay->GetElement(OverlayEls.C)));
+				}
 			}
 			else
 			{
-				InstanceAttrib.Set(InstanceTri[0], AttribIndex, DefaultValue);
-				InstanceAttrib.Set(InstanceTri[1], AttribIndex, DefaultValue);
-				InstanceAttrib.Set(InstanceTri[2], AttribIndex, DefaultValue);
+				SetDefaults();
 			}
 		}
 	}
